Add experimental CustomLaunchUI plugin
- configurable via editor Experimental settings
- currently enabled by default

diff --git a/Engine/Config/BaseEditorPerProjectUserSettings.ini b/Engine/Config/BaseEditorPerProjectUserSettings.ini
--- a/Engine/Config/BaseEditorPerProjectUserSettings.ini
+++ b/Engine/Config/BaseEditorPerProjectUserSettings.ini
@@ -156,6 +156,7 @@ bToolbarCustomization=False
 bBehaviorTreeEditor=False
 bEnableFindAndReplaceReferences =False
 bExampleLayersAndBlends = True
+bShowNewLaunchProfileTools=True
 
 ; Enabled by default now
 bEnableAsyncTextureCompilation=True
diff --git a/Engine/Plugins/Developer/CustomLaunchUI/CustomLaunchUI.uplugin b/Engine/Plugins/Developer/CustomLaunchUI/CustomLaunchUI.uplugin
new file mode 100644
--- /dev/null
+++ b/Engine/Plugins/Developer/CustomLaunchUI/CustomLaunchUI.uplugin
@@ -0,0 +1,27 @@
+{
+	"FileVersion": 3,
+	"Version": 1,
+	"VersionName": "0.1",
+	"FriendlyName": "Custom Launch UI",
+	"Description": "Configure custom project launch profiles.",
+	"CreatedBy": "Epic Games, Inc.",
+	"Category": "Platform",
+	"CreatedByURL": "https://epicgames.com",
+	"DocsURL": "",
+	"MarketplaceURL": "",
+	"SupportURL": "",
+	"EnabledByDefault": true,
+	"RequiresBuildPlatform": true,
+	"CanContainContent": false,
+	"IsBetaVersion": true,
+	"Installed": false,
+	"SupportedTargetPlatforms": [ "Win64" ],
+	"Modules": [
+		{
+			"Name": "CustomLaunchUI",
+			"Type": "Editor",
+			"LoadingPhase": "Default",
+			"PlatformAllowList": [ "Win64" ]
+		}
+	]
+}
\ No newline at end of file
diff --git a/Engine/Plugins/Developer/CustomLaunchUI/Source/CustomLaunchUI.Build.cs b/Engine/Plugins/Developer/CustomLaunchUI/Source/CustomLaunchUI.Build.cs
new file mode 100644
--- /dev/null
+++ b/Engine/Plugins/Developer/CustomLaunchUI/Source/CustomLaunchUI.Build.cs
@@ -0,0 +1,28 @@
+// Copyright Epic Games, Inc. All Rights Reserved.
+
+using UnrealBuildTool;
+
+public class CustomLaunchUI : ModuleRules
+{
+	public CustomLaunchUI(ReadOnlyTargetRules Target) : base(Target)
+	{
+		PrivateDependencyModuleNames.AddRange(
+		[
+			"Core",
+			"CoreUObject",
+			"LauncherServices",
+			"TargetPlatform",
+			"DesktopPlatform",
+			"ApplicationCore",
+			"InputCore",
+			"Slate",
+			"SlateCore",
+			"WorkspaceMenuStructure",
+			"ToolWidgets",
+			"TargetDeviceServices",
+			"PropertyEditor",
+			"Projects",
+			"DeveloperToolSettings",
+		]);
+	}
+}
diff --git a/Engine/Plugins/Developer/CustomLaunchUI/Source/Private/Classes/LauncherProfileSettings.cpp b/Engine/Plugins/Developer/CustomLaunchUI/Source/Private/Classes/LauncherProfileSettings.cpp
new file mode 100644
--- /dev/null
+++ b/Engine/Plugins/Developer/CustomLaunchUI/Source/Private/Classes/LauncherProfileSettings.cpp
@@ -0,0 +1,404 @@
+// Copyright Epic Games, Inc. All Rights Reserved.
+
+#include "LauncherProfileSettings.h"
+#include "GameProjectHelper.h"
+#include "PlatformInfo.h"
+
+#define LOCTEXT_NAMESPACE "LauncherProfileSettings"
+
+/* @todo: which is better
+
+   LAUNCHER_USE_GLOBAL_PROJECTPATH=0
+   Each device profile manages its own project & build target.
+   This means that in editor "Current Project" is always the loaded project and cannot be changed
+   In UnrealFrontend, "Current Project" will be selectable via an additional toolbar @todo
+   Custom Profiles are not affected
+
+   LAUNCHER_USE_GLOBAL_PROJECTPATH=1
+   Each device profile shares the global project & build target.
+   In editor, changing the project/target in one device will affect all device profiles - and, crucially, all Custom Profiles that use "Current Project"
+   This is how Project Launcher operated - however, Project Launcher presented this as "Any Project" rather than "Current Project" & the global project combo was always visible
+
+*/
+#define LAUNCHER_USE_GLOBAL_PROJECTPATH 0
+
+namespace Private
+{
+	ELaunchProfileBuildConfiguration Translate( EBuildConfiguration Configuration )
+	{
+		switch (Configuration)
+		{
+			case EBuildConfiguration::Debug:       return ELaunchProfileBuildConfiguration::Debug;
+			case EBuildConfiguration::Development: return ELaunchProfileBuildConfiguration::Development;
+			case EBuildConfiguration::Shipping:    return ELaunchProfileBuildConfiguration::Shipping;
+			case EBuildConfiguration::Test:        return ELaunchProfileBuildConfiguration::Test;
+		}
+		return ELaunchProfileBuildConfiguration::Development;
+	}
+
+	EBuildConfiguration Translate( ELaunchProfileBuildConfiguration Configuration )
+	{
+		switch (Configuration)
+		{
+			case ELaunchProfileBuildConfiguration::Debug:       return EBuildConfiguration::Debug;
+			case ELaunchProfileBuildConfiguration::Development: return EBuildConfiguration::Development;
+			case ELaunchProfileBuildConfiguration::Shipping:    return EBuildConfiguration::Shipping;
+			case ELaunchProfileBuildConfiguration::Test:        return EBuildConfiguration::Test;
+		}
+		return EBuildConfiguration::Unknown;
+	}
+}
+
+
+void UBaseLauncherProfileSettings::SetProfile( TSharedPtr<CustomLaunchUI::FModel> InModel, ILauncherProfilePtr InProfile )
+{
+	Model = InModel;
+	Profile = InProfile;
+}
+
+void UBaseLauncherProfileSettings::RefreshProfile()
+{
+	SetProfile(Model, Profile);
+}
+
+
+
+
+
+
+ULauncherBasicLaunchProfileSettings::ULauncherBasicLaunchProfileSettings(const FObjectInitializer& ObjectInitializer)
+	: Super(ObjectInitializer)
+{
+}
+
+
+void ULauncherBasicLaunchProfileSettings::SetProfile( TSharedPtr<CustomLaunchUI::FModel> InModel, ILauncherProfilePtr InProfile )
+{
+	Super::SetProfile(InModel, InProfile);
+
+	ProjectPath = Profile->HasProjectSpecified() ? Profile->GetProjectPath() : FString();
+	bHasProjectSpecified = Profile->HasProjectSpecified();
+	Target = Profile->HasBuildTargetSpecified() ? Profile->GetBuildTarget() : FString();
+	bHasTargetSpecified = Profile->HasBuildTargetSpecified();
+
+	Configuration = Private::Translate(Profile->GetBuildConfiguration());
+	ContentScheme = (uint8)Model->DetermineProfileContentScheme(Profile.ToSharedRef());
+	CommandLine = Profile->GetAdditionalCommandLineParameters();
+
+	TArray<FString> DeployDeviceIDs = InProfile->GetDeployedDeviceGroup()->GetDeviceIDs();
+	DeployDeviceID = DeployDeviceIDs.Num() > 0 ? DeployDeviceIDs[0] : FString();
+}
+
+void ULauncherBasicLaunchProfileSettings::PostEditChangeProperty(struct FPropertyChangedEvent& PropertyChangedEvent)
+{
+	FProperty* Property = PropertyChangedEvent.MemberProperty ? PropertyChangedEvent.MemberProperty : PropertyChangedEvent.Property;
+	if (Property == nullptr)
+	{
+		return;
+	}
+
+    const FName PropertyName = Property->GetFName();
+
+    if (PropertyName == GET_MEMBER_NAME_CHECKED(ULauncherBasicLaunchProfileSettings, ProjectPath))
+	{
+		bHasProjectSpecified = !ProjectPath.IsEmpty();
+		Profile->SetProjectPath(ProjectPath);
+		Profile->SetProjectSpecified(bHasProjectSpecified);
+
+		// updating the project can reset the Taret
+		Target = Profile->HasBuildTargetSpecified() ? Profile->GetBuildTarget() : FString();
+		bHasTargetSpecified = Profile->HasBuildTargetSpecified();
+	}
+    else if (PropertyName == GET_MEMBER_NAME_CHECKED(ULauncherBasicLaunchProfileSettings, Target))
+	{
+		bHasTargetSpecified = !Target.IsEmpty();
+		Profile->SetBuildTarget(Target);
+		Profile->SetBuildTargetSpecified(bHasTargetSpecified);
+
+		Model->UpdatedCookedPlatformsFromDeployDeviceProxy(Profile.ToSharedRef());
+	}
+    else if (PropertyName == GET_MEMBER_NAME_CHECKED(ULauncherBasicLaunchProfileSettings, Configuration))
+	{
+		Profile->SetBuildConfiguration(Private::Translate(Configuration));
+	}
+    else if (PropertyName == GET_MEMBER_NAME_CHECKED(ULauncherBasicLaunchProfileSettings, ContentScheme))
+	{
+		Model->SetProfileContentScheme((CustomLaunchUI::EContentScheme)ContentScheme, Profile.ToSharedRef(), ELauncherProfileDeploymentModes::CopyToDevice );
+	}
+    else if (PropertyName == GET_MEMBER_NAME_CHECKED(ULauncherBasicLaunchProfileSettings, CommandLine))
+	{
+		Profile->SetAdditionalCommandLineParameters(CommandLine);
+	}
+	else if (PropertyName == GET_MEMBER_NAME_CHECKED(ULauncherBasicLaunchProfileSettings, DeployDeviceID))
+	{
+		Profile->GetDeployedDeviceGroup()->RemoveAllDevices();
+		if (!DeployDeviceID.IsEmpty())
+		{
+			Profile->GetDeployedDeviceGroup()->AddDevice(DeployDeviceID);
+			Model->UpdatedCookedPlatformsFromDeployDeviceProxy(Profile.ToSharedRef());
+		}
+	}
+}
+
+
+
+
+
+
+
+
+
+ULauncherCustomProfileSettings::ULauncherCustomProfileSettings(const FObjectInitializer& ObjectInitializer)
+	: Super(ObjectInitializer)
+{
+}
+
+void ULauncherCustomProfileSettings::SetProfile( TSharedPtr<CustomLaunchUI::FModel> InModel, ILauncherProfilePtr InProfile )
+{
+	Super::SetProfile(InModel, InProfile);
+
+	// -- general settings --
+	ProjectPath = Profile->HasProjectSpecified() ? Profile->GetProjectPath() : FString();
+	bHasProjectSpecified = Profile->HasProjectSpecified();
+	Target = Profile->HasBuildTargetSpecified() && !Profile->GetBuildTarget().IsEmpty() ? Profile->GetBuildTarget() : FString();
+	bHasTargetSpecified = Profile->HasBuildTargetSpecified();
+	Platforms = Profile->GetCookedPlatforms();
+	Configuration = Private::Translate(Profile->GetBuildConfiguration());
+
+	// -- source content options --
+	ContentScheme = (uint8)Model->DetermineProfileContentScheme(Profile.ToSharedRef());
+	bCompressPakFiles = Profile->IsCompressed();
+	bUseIoStore = Profile->IsUsingIoStore();
+	bGenerateChunks = Profile->IsGeneratingChunks();
+	bImportBestMatchZenSnapshot = Profile->IsImportingZenSnapshot();
+	ZenPakStreamingPath.Path = Profile->GetZenPakStreamingPath();
+	bIncrementalCook = Profile->IsCookingIncrementally();
+	AdditionalCookerOptions = Profile->GetCookOptions();
+	MapsToCook = Profile->GetCookedMaps();
+
+	// -- stage options --
+	CustomStageDirectory.Path = Profile->GetPackageDirectory();
+	bAchiveBuild = Profile->IsArchiving();
+	ArchiveDirectory.Path = Profile->GetArchiveDirectory();
+
+	// -- code options --
+	bBuildCode = (Profile->GetBuildMode() != ELauncherProfileBuildModes::DoNotBuild);
+	bBuildAlways = (Profile->GetBuildMode() == ELauncherProfileBuildModes::Build);
+	bBuildUAT = Profile->IsBuildingUAT();
+
+	// -- deploy and run options ---
+	bDeployToDevice = Profile->GetDeploymentMode() != ELauncherProfileDeploymentModes::DoNotDeploy;
+	bOnlyDeployModifiedContent = Profile->IsDeployingIncrementally();
+	if (Profile->GetDeployedDeviceGroup() != nullptr)
+	{
+		DeployDeviceIDs = Profile->GetDeployedDeviceGroup()->GetDeviceIDs();
+	}
+	bRun = Profile->GetLaunchMode() != ELauncherProfileLaunchModes::DoNotLaunch;
+	InitialMap = Profile->GetDefaultLaunchRole()->GetInitialMap(); // fixme: map picker UI
+
+	// get unified command line from these two fields. the first is presented in "Build" for old Project Launcher, and the latter is presented in "Launch" for old Project Launcher)
+	// when we save back to the profile, this will be stored just in the "Build" one for clarity (because multiple roles are not supported in old or new project launcher)
+	CommandLine = Profile->GetAdditionalCommandLineParameters().TrimStartAndEnd() + TEXT(" ") + Profile->GetDefaultLaunchRole()->GetUATCommandLine().TrimStartAndEnd();
+	CommandLine.TrimStartAndEndInline();
+
+
+	// set up additional items for custom list pickers
+	MapOption = MapsToCook.Num() > 0 ? EMapOption::Selected : EMapOption::Startup;
+	CachedMapsToCook.Reset();
+
+	DeployDeviceOption = DeployDeviceIDs.Num() > 0 ? EDeployDeviceOption::Selected : EDeployDeviceOption::Default;
+	CachedDeployDeviceIDs.Reset();	
+}
+
+void ULauncherCustomProfileSettings::PostEditChangeProperty(struct FPropertyChangedEvent& PropertyChangedEvent)
+{
+	FProperty* Property = PropertyChangedEvent.MemberProperty ? PropertyChangedEvent.MemberProperty : PropertyChangedEvent.Property;
+	if (Property == nullptr)
+	{
+		return;
+	}
+
+    const FName PropertyName = Property->GetFName();
+
+    if (PropertyName == GET_MEMBER_NAME_CHECKED(ULauncherCustomProfileSettings, ProjectPath))
+	{
+		Profile->SetProjectPath(ProjectPath);
+		Profile->SetProjectSpecified(bHasProjectSpecified);
+
+		// updating the project can reset the Taret
+		Target = Profile->HasBuildTargetSpecified() && !Profile->GetBuildTarget().IsEmpty() ? Profile->GetBuildTarget() : FString();
+		bHasTargetSpecified = Profile->HasBuildTargetSpecified();
+
+		if (CustomLaunchUI::bUseFriendlyBuildTargetSelection)
+		{
+			Model->UpdateCookedPlatformsFromBuildTarget(Profile.ToSharedRef());
+			Platforms = Profile->GetCookedPlatforms();
+		}
+	}
+    else if (PropertyName == GET_MEMBER_NAME_CHECKED(ULauncherCustomProfileSettings, Target))
+	{
+		bHasTargetSpecified = !Target.IsEmpty();
+		Profile->SetBuildTarget( bHasTargetSpecified ? Target : FString());
+		Profile->SetBuildTargetSpecified(bHasTargetSpecified);
+
+		if (CustomLaunchUI::bUseFriendlyBuildTargetSelection)
+		{
+			Model->UpdateCookedPlatformsFromBuildTarget(Profile.ToSharedRef());
+			Platforms = Profile->GetCookedPlatforms();
+		}
+	}
+    else if (PropertyName == GET_MEMBER_NAME_CHECKED(ULauncherCustomProfileSettings, Platforms))
+	{
+		Profile->ClearCookedPlatforms();
+		for (const FString& Platform : Platforms)
+		{
+			Profile->AddCookedPlatform(Platform);
+		}
+	}
+    else if (PropertyName == GET_MEMBER_NAME_CHECKED(ULauncherCustomProfileSettings, Configuration))
+	{
+		Profile->SetBuildConfiguration(Private::Translate(Configuration));
+	}
+    else if (PropertyName == GET_MEMBER_NAME_CHECKED(ULauncherCustomProfileSettings, bBuildCode) || PropertyName == GET_MEMBER_NAME_CHECKED(ULauncherCustomProfileSettings, bBuildAlways))
+	{ 
+		if (!bBuildCode)
+		{
+			Profile->SetBuildMode(ELauncherProfileBuildModes::DoNotBuild);
+		}
+		else if (bBuildAlways)
+		{
+			Profile->SetBuildMode(ELauncherProfileBuildModes::Build);
+		}
+		else
+		{
+			Profile->SetBuildMode(ELauncherProfileBuildModes::Auto);
+		}
+	}
+    else if (PropertyName == GET_MEMBER_NAME_CHECKED(ULauncherCustomProfileSettings, bBuildUAT))
+	{
+		Profile->SetBuildUAT(bBuildUAT);
+	}
+    else if (PropertyName == GET_MEMBER_NAME_CHECKED(ULauncherCustomProfileSettings, ContentScheme))
+	{
+		Model->SetProfileContentScheme((CustomLaunchUI::EContentScheme)ContentScheme, Profile.ToSharedRef(), bDeployToDevice ? ELauncherProfileDeploymentModes::CopyToDevice : ELauncherProfileDeploymentModes::DoNotDeploy );
+	}
+    else if (PropertyName == GET_MEMBER_NAME_CHECKED(ULauncherCustomProfileSettings, bCompressPakFiles))
+	{
+		Profile->SetCompressed(bCompressPakFiles);
+	}
+    else if (PropertyName == GET_MEMBER_NAME_CHECKED(ULauncherCustomProfileSettings, bUseIoStore))
+	{
+		Profile->SetUseIoStore(bUseIoStore);
+	}
+    else if (PropertyName == GET_MEMBER_NAME_CHECKED(ULauncherCustomProfileSettings, bGenerateChunks))
+	{
+		Profile->SetGenerateChunks(bGenerateChunks);
+	}
+    else if (PropertyName == GET_MEMBER_NAME_CHECKED(ULauncherCustomProfileSettings, bImportBestMatchZenSnapshot))
+	{
+		Profile->SetImportingZenSnapshot(bImportBestMatchZenSnapshot);
+	}
+    else if (PropertyName == GET_MEMBER_NAME_CHECKED(ULauncherCustomProfileSettings, ZenPakStreamingPath))
+	{
+		Profile->SetZenPakStreamingPath(ZenPakStreamingPath.Path);
+	}
+    else if (PropertyName == GET_MEMBER_NAME_CHECKED(ULauncherCustomProfileSettings, bIncrementalCook))
+	{
+		Profile->SetIncrementalCooking(bIncrementalCook);
+	}
+    else if (PropertyName == GET_MEMBER_NAME_CHECKED(ULauncherCustomProfileSettings, AdditionalCookerOptions))
+	{
+		Profile->SetCookOptions(AdditionalCookerOptions);
+	}
+	else if (PropertyName == GET_MEMBER_NAME_CHECKED(ULauncherCustomProfileSettings, MapsToCook))
+	{
+		// fixme: do better
+		Profile->ClearCookedMaps(); 
+		for (const FString& Map : MapsToCook)
+		{
+			Profile->AddCookedMap(Map);
+		}
+	} 
+    else if (PropertyName == GET_MEMBER_NAME_CHECKED(ULauncherCustomProfileSettings, CustomStageDirectory))
+	{
+		Profile->SetPackageDirectory(CustomStageDirectory.Path);
+	}
+    else if (PropertyName == GET_MEMBER_NAME_CHECKED(ULauncherCustomProfileSettings, bAchiveBuild))
+	{
+		Profile->SetArchive(bAchiveBuild);
+	}
+    else if (PropertyName == GET_MEMBER_NAME_CHECKED(ULauncherCustomProfileSettings, ArchiveDirectory))
+	{
+		Profile->SetArchiveDirectory(ArchiveDirectory.Path);
+	}
+    else if (PropertyName == GET_MEMBER_NAME_CHECKED(ULauncherCustomProfileSettings, bDeployToDevice))
+	{
+		if ((CustomLaunchUI::EContentScheme)ContentScheme != CustomLaunchUI::EContentScheme::CookOnTheFly)
+		{
+			Profile->SetDeploymentMode(bDeployToDevice ? ELauncherProfileDeploymentModes::CopyToDevice : ELauncherProfileDeploymentModes::DoNotDeploy );
+		}
+	}
+    else if (PropertyName == GET_MEMBER_NAME_CHECKED(ULauncherCustomProfileSettings, bOnlyDeployModifiedContent))
+	{
+		Profile->SetIncrementalDeploying(bOnlyDeployModifiedContent);
+	}
+    else if (PropertyName == GET_MEMBER_NAME_CHECKED(ULauncherCustomProfileSettings, DeployDeviceIDs))
+	{
+		// fixme: do better?
+		if (Profile->GetDeployedDeviceGroup() == nullptr)
+		{
+			Profile->SetDeployedDeviceGroup(Model->GetProfileManager()->AddNewDeviceGroup());
+		}
+
+		Profile->GetDeployedDeviceGroup()->RemoveAllDevices(); 
+		for ( const FString& DeviceID : DeployDeviceIDs )
+		{
+			Profile->GetDeployedDeviceGroup()->AddDevice(DeviceID);
+		}
+	}
+    else if (PropertyName == GET_MEMBER_NAME_CHECKED(ULauncherCustomProfileSettings, bRun))
+	{
+		Profile->SetLaunchMode(bRun ? ELauncherProfileLaunchModes::DefaultRole : ELauncherProfileLaunchModes::DoNotLaunch);
+	}
+    else if (PropertyName == GET_MEMBER_NAME_CHECKED(ULauncherCustomProfileSettings, CommandLine))
+	{
+		// unify command line to the primary field because multi-role deploy is not supported in any tool
+		Profile->SetAdditionalCommandLineParameters(CommandLine);
+		Profile->GetDefaultLaunchRole()->SetCommandLine(TEXT(""));
+	}
+    else if (PropertyName == GET_MEMBER_NAME_CHECKED(ULauncherCustomProfileSettings, InitialMap))
+	{
+		Profile->GetDefaultLaunchRole()->SetInitialMap(InitialMap );
+	}
+
+	Model->GetProfileManager()->SaveJSONProfile(Profile.ToSharedRef());
+}
+
+
+bool ULauncherCustomProfileSettings::ForPak()
+{
+	return (CustomLaunchUI::EContentScheme)ContentScheme == CustomLaunchUI::EContentScheme::PakFiles || (CustomLaunchUI::EContentScheme)ContentScheme == CustomLaunchUI::EContentScheme::DevelopmentPackage;
+}
+
+bool ULauncherCustomProfileSettings::ForZenWS()
+{
+	return (CustomLaunchUI::EContentScheme)ContentScheme == CustomLaunchUI::EContentScheme::ZenPakStreaming;
+}
+
+bool ULauncherCustomProfileSettings::ForCooked()
+{
+	return (CustomLaunchUI::EContentScheme)ContentScheme != CustomLaunchUI::EContentScheme::ZenPakStreaming && (CustomLaunchUI::EContentScheme)ContentScheme != CustomLaunchUI::EContentScheme::CookOnTheFly;
+}
+
+bool ULauncherCustomProfileSettings::NeedsProject()
+{
+	return bHasProjectSpecified && !ProjectPath.IsEmpty();
+}
+
+
+TArray<FString> ULauncherCustomProfileSettings::GetAvailableMaps() const
+{
+	return FGameProjectHelper::GetAvailableMaps(Profile->GetProjectBasePath(), Profile->SupportsEngineMaps(), true);
+}
+
+#undef LOCTEXT_NAMESPACE
diff --git a/Engine/Plugins/Developer/CustomLaunchUI/Source/Private/Classes/LauncherProfileSettings.h b/Engine/Plugins/Developer/CustomLaunchUI/Source/Private/Classes/LauncherProfileSettings.h
new file mode 100644
--- /dev/null
+++ b/Engine/Plugins/Developer/CustomLaunchUI/Source/Private/Classes/LauncherProfileSettings.h
@@ -0,0 +1,202 @@
+// Copyright Epic Games, Inc. All Rights Reserved.
+
+#pragma once
+
+#include "CoreMinimal.h"
+#include "UObject/Object.h"
+#include "UObject/SoftObjectPath.h"
+#include "ILauncherProfile.h"
+#include "Model/CustomLaunchUIModel.h"
+#include "LauncherProfileSettings.generated.h"
+
+
+UENUM()
+enum class ELaunchProfileBuildConfiguration : uint8
+{
+	Debug,
+	Development,
+	Shipping,
+	Test
+};
+
+UCLASS(Abstract)
+class UBaseLauncherProfileSettings : public UObject
+{
+	GENERATED_BODY()
+
+public:
+	virtual void SetProfile( TSharedPtr<CustomLaunchUI::FModel> InModel, ILauncherProfilePtr InProfile );
+	void RefreshProfile();
+
+protected:
+	ILauncherProfilePtr Profile;
+	TSharedPtr<CustomLaunchUI::FModel> Model;
+};
+
+
+
+UCLASS()
+class ULauncherBasicLaunchProfileSettings : public UBaseLauncherProfileSettings
+{
+public:
+
+	GENERATED_UCLASS_BODY()
+
+	UPROPERTY(EditAnywhere, Category = "General Settings", DisplayName="Project")
+	FString ProjectPath;
+	bool bHasProjectSpecified;
+
+	UPROPERTY(EditAnywhere, Category = "General Settings")
+	FString Target;
+	bool bHasTargetSpecified;
+
+	UPROPERTY(EditAnywhere, Category = "General Settings")
+	ELaunchProfileBuildConfiguration Configuration = ELaunchProfileBuildConfiguration::Development;
+
+	UPROPERTY(EditAnywhere, Category = "General Settings", Meta=(DisplayName="Content Scheme"))
+	uint8 ContentScheme = (uint8)CustomLaunchUI::EContentScheme::ZenStreaming;
+
+	UPROPERTY(EditAnywhere, Category = "General Settings")
+	FString CommandLine;
+
+	UPROPERTY(EditAnywhere, Category = "General Settings", Meta=(DisplayName="Target Device"))
+	FString DeployDeviceID;
+
+	virtual void SetProfile( TSharedPtr<CustomLaunchUI::FModel> InModel, ILauncherProfilePtr InProfile ) override;
+
+private:
+	virtual void PostEditChangeProperty(struct FPropertyChangedEvent& PropertyChangedEvent) override;
+
+};
+
+
+UCLASS()
+class ULauncherCustomProfileSettings : public UBaseLauncherProfileSettings
+{
+public:
+
+	GENERATED_UCLASS_BODY()
+
+	UPROPERTY(EditAnywhere, Category = "General Settings", Meta=(DisplayName="Project"))
+	FString ProjectPath;
+	bool bHasProjectSpecified;
+
+	UPROPERTY(EditAnywhere, Category = "General Settings", Meta=(EditCondition="NeedsProject()"))
+	FString Target;
+	bool bHasTargetSpecified;
+
+	UPROPERTY(EditAnywhere, Category = "General Settings", Meta=(DisplayName="Platform"))
+	TArray<FString> Platforms;
+
+	UPROPERTY(EditAnywhere, Category = "General Settings", Meta=(EditCondition=bBuildCode))
+	ELaunchProfileBuildConfiguration Configuration = ELaunchProfileBuildConfiguration::Development;
+
+	UPROPERTY(EditAnywhere, Category = "General Settings", Meta=(DisplayName="Content Scheme"))
+	uint8 ContentScheme = (uint8)CustomLaunchUI::EContentScheme::ZenStreaming;
+
+
+
+	UPROPERTY(EditAnywhere, Category = "Content Scheme Properties", Meta=(EditConditionHides, EditCondition="ForPak()"))
+	bool bCompressPakFiles;
+
+	UPROPERTY(EditAnywhere, Category = "Content Scheme Properties", Meta=(EditConditionHides, EditCondition="ForPak()"))
+	bool bUseIoStore;
+
+	UPROPERTY(EditAnywhere, Category = "Content Scheme Properties", Meta=(EditConditionHides, EditCondition="ForPak()"))
+	bool bGenerateChunks;
+
+	UPROPERTY(EditAnywhere, Category = "Content Scheme Properties", Meta=(EditConditionHides, EditCondition="ForCooked()"))
+	bool bImportBestMatchZenSnapshot;
+
+	UPROPERTY(EditAnywhere, Category = "Content Scheme Properties")
+	FDirectoryPath ZenPakStreamingPath;  // nb. EditCondition with EditConditionHides does not seem to work with FDirectoryPath. this field is hidden in LauncherProfileSettingsDetails.cpp instead
+
+	UPROPERTY(EditAnywhere, Category = "Content Scheme Properties", Meta=(EditConditionHides, EditCondition="ForCooked()"))
+	bool bIncrementalCook;
+
+
+	UPROPERTY(EditAnywhere, Category = "Directory", Meta=(EditCondition="ForCooked()")/*, AdvancedDisplay*/)
+	FDirectoryPath CustomStageDirectory;
+
+	UPROPERTY(EditAnywhere, Category = "Directory", Meta=(EditCondition="ForCooked()"))
+	bool bAchiveBuild;
+
+	UPROPERTY(EditAnywhere, Category = "Directory", Meta=(EditCondition="bAchiveBuild && ForCooked()"))
+	FDirectoryPath ArchiveDirectory;
+
+
+	UPROPERTY(EditAnywhere, Category = "Maps", Meta=(EditCondition="ForCooked()"))
+	TArray<FString> MapsToCook;
+	TArray<FString> CachedMapsToCook;
+	enum class EMapOption : uint8
+	{
+		Startup,
+		Selected,
+	};
+	EMapOption MapOption;
+
+	UPROPERTY(EditAnywhere, Category = "Maps", Meta=(EditCondition="ForCooked()")/*, AdvancedDisplay*/)
+	FString AdditionalCookerOptions;
+
+
+
+	UPROPERTY(EditAnywhere, Category = "Code")
+	bool bBuildCode;
+
+	UPROPERTY(EditAnywhere, Category = "Code", Meta=(EditCondition=bBuildCode,DisplayName="Build if pre-build target exists"))
+	bool bBuildAlways;
+
+	UPROPERTY(EditAnywhere, Category = "Code", Meta=(EditCondition=bBuildCode,DisplayName="Build UAT")/*, AdvancedDisplay*/)
+	bool bBuildUAT;
+
+
+
+	UPROPERTY(EditAnywhere, Category = "Deploy and Run")
+	bool bDeployToDevice;
+
+	UPROPERTY(EditAnywhere, Category = "Deploy and Run", Meta=(EditCondition=bDeployToDevice)) // fixme: "only linux & android" ?
+	bool bOnlyDeployModifiedContent;
+
+	UPROPERTY(EditAnywhere, Category = "Deploy and Run", Meta=(EditCondition=bDeployToDevice, DisplayName="Target Device"))
+	TArray<FString> DeployDeviceIDs;
+	TArray<FString> CachedDeployDeviceIDs;
+	enum class EDeployDeviceOption : uint8
+	{
+		Default,
+		Selected,
+	};
+	EDeployDeviceOption DeployDeviceOption;
+
+	UPROPERTY(EditAnywhere, Category = "Deploy and Run")
+	bool bRun;
+
+	UPROPERTY(EditAnywhere, Category = "Deploy and Run", Meta=(EditCondition=bRun, GetOptions="GetAvailableMaps")) // fixme: map combo / better UI
+	FString InitialMap;
+
+	UPROPERTY(EditAnywhere, Category = "Deploy and Run", Meta=(EditCondition=bRun))
+	FString CommandLine;
+
+
+
+	void SetProfile( TSharedPtr<CustomLaunchUI::FModel> InModel, ILauncherProfilePtr InProfile ) override;
+
+private:
+	virtual void PostEditChangeProperty(struct FPropertyChangedEvent& PropertyChangedEvent) override;
+
+
+	UFUNCTION()
+	bool ForPak();
+
+	UFUNCTION()
+	bool ForZenWS();
+
+	UFUNCTION()
+	bool ForCooked();
+
+	UFUNCTION()
+	bool NeedsProject();
+
+	UFUNCTION()
+	TArray<FString> GetAvailableMaps() const;
+
+};
diff --git a/Engine/Plugins/Developer/CustomLaunchUI/Source/Private/CustomLaunchUIModule.cpp b/Engine/Plugins/Developer/CustomLaunchUI/Source/Private/CustomLaunchUIModule.cpp
new file mode 100644
--- /dev/null
+++ b/Engine/Plugins/Developer/CustomLaunchUI/Source/Private/CustomLaunchUIModule.cpp
@@ -0,0 +1,82 @@
+// Copyright Epic Games, Inc. All Rights Reserved.
+
+#include "Modules/ModuleManager.h"
+#include "WorkspaceMenuStructure.h"
+#include "WorkspaceMenuStructureModule.h"
+#include "Framework/Docking/TabManager.h"
+#include "Widgets/Docking/SDockTab.h"
+#include "Widgets/SCustomLaunchUI.h"
+#include "Styling/CustomLaunchUIStyle.h"
+#include "Details/LauncherProfileSettingsDetails.h"
+#include "ILauncherServicesModule.h"
+#include "ITargetDeviceServicesModule.h"
+#include "PropertyEditorModule.h"
+#include "Model/CustomLaunchUIModel.h"
+#include "Classes/LauncherProfileSettings.h"
+
+#define LOCTEXT_NAMESPACE "FCustomLaunchUIModule"
+
+static const FName TabName("CustomLaunchUI");
+
+class FCustomLaunchUIModule : public IModuleInterface
+{
+public:
+
+	virtual void StartupModule() override
+	{
+		// register detail customization
+		static FName PropertyEditor("PropertyEditor");
+		FPropertyEditorModule* PropertyModule = FModuleManager::GetModulePtr<FPropertyEditorModule>(PropertyEditor);
+
+		if (PropertyModule != nullptr)
+		{
+			static FName LauncherBasicLaunchProfileSettings("LauncherBasicLaunchProfileSettings");
+			PropertyModule->RegisterCustomClassLayout(LauncherBasicLaunchProfileSettings, FOnGetDetailCustomizationInstance::CreateStatic(&FLauncherBasicLaunchProfileSettingsDetails::MakeInstance));
+
+			static FName LauncherCustomProfileSettings("LauncherCustomProfileSettings");
+			PropertyModule->RegisterCustomClassLayout(LauncherCustomProfileSettings, FOnGetDetailCustomizationInstance::CreateStatic(&FLauncherCustomProfileSettingsDetails::MakeInstance));
+		}
+
+		// register custom styles
+		FCustomLaunchUIStyle::Initialize();
+
+		// register main tab
+		FGlobalTabmanager::Get()->RegisterNomadTabSpawner(TabName, FOnSpawnTab::CreateRaw(this, &FCustomLaunchUIModule::HandleSpawnTab))
+			.SetDisplayName(LOCTEXT("SpawnTabLabel", "Custom Launch Profiles"))
+			.SetTooltipText(LOCTEXT("SpawnTabToolTip", "Configure custom launch profiles for advanced packaging, deploying and launching of your projects"))
+			.SetIcon(FSlateIcon(FAppStyle::GetAppStyleSetName(), "Launcher.TabIcon"))
+			.SetGroup(WorkspaceMenu::GetMenuStructure().GetDeveloperToolsPlatformsCategory())
+			.SetMenuType(ETabSpawnerMenuType::Hidden); // hide it for now - it will be invoked by TurnkeySupportModule
+	}
+
+	virtual void ShutdownModule() override
+	{
+		FGlobalTabmanager::Get()->UnregisterNomadTabSpawner(TabName);
+		FCustomLaunchUIStyle::Shutdown();
+	}
+
+private:
+
+	TSharedRef<SDockTab> HandleSpawnTab(const FSpawnTabArgs& SpawnTabArgs)
+	{
+		ILauncherServicesModule& LauncherServicesModule = FModuleManager::LoadModuleChecked<ILauncherServicesModule>("LauncherServices");
+		ITargetDeviceServicesModule& TargetDeviceServicesModule = FModuleManager::LoadModuleChecked<ITargetDeviceServicesModule>("TargetDeviceServices");
+
+		TSharedRef<CustomLaunchUI::FModel> Model = MakeShared<CustomLaunchUI::FModel>(TargetDeviceServicesModule.GetDeviceProxyManager(), LauncherServicesModule.CreateLauncher(), LauncherServicesModule.GetProfileManager() );
+		
+		// configure CDOs, for 'reset to default'
+		ULauncherBasicLaunchProfileSettings::StaticClass()->GetDefaultObject<ULauncherBasicLaunchProfileSettings>()->SetProfile(Model, Model->CreateBasicLaunchProfile());
+		ULauncherCustomProfileSettings::StaticClass()->     GetDefaultObject<ULauncherCustomProfileSettings>()->     SetProfile(Model, Model->CreateCustomProfile( TEXT("DefaultCustomProfile") ));
+
+		return SNew(SDockTab)
+		.TabRole(ETabRole::NomadTab)
+		[
+			SNew(SCustomLaunchUI, Model)
+		];
+	}
+};
+
+
+IMPLEMENT_MODULE(FCustomLaunchUIModule, CustomLaunchUI);
+
+#undef LOCTEXT_NAMESPACE
diff --git a/Engine/Plugins/Developer/CustomLaunchUI/Source/Private/Details/LauncherProfileSettingsDetails.cpp b/Engine/Plugins/Developer/CustomLaunchUI/Source/Private/Details/LauncherProfileSettingsDetails.cpp
new file mode 100644
--- /dev/null
+++ b/Engine/Plugins/Developer/CustomLaunchUI/Source/Private/Details/LauncherProfileSettingsDetails.cpp
@@ -0,0 +1,1076 @@
+// Copyright Epic Games, Inc. All Rights Reserved.
+
+#include "LauncherProfileSettingsDetails.h"
+
+#include "DetailWidgetRow.h"
+#include "DetailCategoryBuilder.h"
+#include "DetailLayoutBuilder.h"
+#include "Framework/Application/SlateApplication.h"
+#include "Framework/MultiBox/MultiBoxBuilder.h"
+#include "Widgets/Shared/SCustomLaunchPlatformCombo.h"
+#include "Widgets/Shared/SCustomLaunchMapListView.h"
+#include "Widgets/Shared/SCustomLaunchDeviceListView.h"
+#include "Widgets/Shared/SCustomLaunchProjectCombo.h"
+#include "Widgets/Shared/SCustomLaunchBuildTargetCombo.h"
+#include "Widgets/Shared/SCustomLaunchContentSchemeCombo.h"
+#include "Widgets/Input/SComboButton.h"
+#include "Widgets/Input/SCheckBox.h"
+#include "Widgets/Input/SSegmentedControl.h"
+#include "Widgets/Input/SMultiLineEditableTextBox.h"
+#include "Widgets/Text/STextBlock.h"
+#include "SResizeBox.h"
+
+#include "SlateOptMacros.h"
+#include "Misc/Paths.h"
+#include "Misc/NotifyHook.h"
+#include "Interfaces/IMainFrameModule.h"
+#include "ILauncherServicesModule.h"
+#include "DesktopPlatformModule.h"
+#include "GameProjectHelper.h"
+#include "PlatformInfo.h"
+#include "Model/CustomLaunchUIModel.h"
+
+#define LOCTEXT_NAMESPACE "LauncherProfileSettingsDetails"
+
+
+
+void FLauncherBasicLaunchProfileSettingsDetails::CustomizeDetails(IDetailLayoutBuilder& DetailBuilder)
+{
+	FBaseLauncherProfileDetails::CustomizeDetails(DetailBuilder);
+
+	// customize project selector
+	TSharedRef<IPropertyHandle> ProjectsPropertyHandle = DetailBuilder.GetProperty(TEXT("ProjectPath"));
+	DetailBuilder.EditDefaultProperty(ProjectsPropertyHandle)->CustomWidget()
+	.OverrideResetToDefault(FResetToDefaultOverride::Hide())
+	.NameContent()	
+	[
+		ProjectsPropertyHandle->CreatePropertyNameWidget()
+	]
+	.ValueContent()
+	.HAlign(HAlign_Fill)
+	[
+		SNew(SCustomLaunchProjectCombo)
+		.OnSelectionChanged(this, &FLauncherBasicLaunchProfileSettingsDetails::SetProjectName)
+		.SelectedProject(this, &FLauncherBasicLaunchProfileSettingsDetails::GetProjectPath)
+		.HasProject(this, &FLauncherBasicLaunchProfileSettingsDetails::HasProject)
+		.CurrentProjectOption(SCustomLaunchProjectCombo::ECurrentProjectOption::Empty)
+		.Font(IDetailLayoutBuilder::GetDetailFont())
+	];
+	
+	// customize the build target combo
+	TSharedRef<IPropertyHandle> TargetPropertyHandle = DetailBuilder.GetProperty(TEXT("Target"));
+	DetailBuilder.EditDefaultProperty(TargetPropertyHandle)->CustomWidget()
+	.OverrideResetToDefault(FResetToDefaultOverride::Hide())
+	.NameContent()	
+	[
+		TargetPropertyHandle->CreatePropertyNameWidget()
+	]
+	.ValueContent()
+	.HAlign(HAlign_Fill)
+	[
+		SNew(SCustomLaunchBuildTargetCombo)
+		.OnSelectionChanged(this, &FLauncherBasicLaunchProfileSettingsDetails::SetBuildTarget)
+		.SelectedBuildTarget(this, &FLauncherBasicLaunchProfileSettingsDetails::GetBuildTarget)
+		.SelectedProject(this, &FLauncherBasicLaunchProfileSettingsDetails::GetProjectPath)
+		.Font(IDetailLayoutBuilder::GetDetailFont())
+	];
+
+	// customize the configuration combo
+	TSharedPtr<IPropertyHandle> ConfigurationPropertyHandle = DetailBuilder.GetProperty(TEXT("Configuration"));
+	DetailBuilder.EditDefaultProperty(ConfigurationPropertyHandle)->CustomWidget()
+	.NameContent()
+	[
+		ConfigurationPropertyHandle->CreatePropertyNameWidget()
+	]
+	.ValueContent()
+	.HAlign(HAlign_Fill)
+	[
+		ConfigurationPropertyHandle->CreatePropertyValueWidget()
+	];
+
+	// customize the content scheme combo
+	TSharedRef<IPropertyHandle> ContentSchemePropertyHandle = DetailBuilder.GetProperty(TEXT("ContentScheme"));
+	DetailBuilder.EditDefaultProperty(ContentSchemePropertyHandle)->CustomWidget()
+	.NameContent()
+	[
+		ContentSchemePropertyHandle->CreatePropertyNameWidget()
+	]
+	.ValueContent()
+	.HAlign(HAlign_Fill)
+	[
+		SNew(SCustomLaunchContentSchemeCombo)
+		.OnSelectionChanged(this, &FLauncherBasicLaunchProfileSettingsDetails::SetContentScheme)
+		.SelectedContentScheme(this, &FLauncherBasicLaunchProfileSettingsDetails::GetContentScheme)
+		.IsContentSchemeAvailable(this, &FLauncherBasicLaunchProfileSettingsDetails::IsContentSchemeAvailable)
+		.Font(IDetailLayoutBuilder::GetDetailFont())
+	];
+
+	// customize the command line textbox
+	TSharedRef<IPropertyHandle> CommandLinePropertyHandle = DetailBuilder.GetProperty(TEXT("CommandLine"));
+	DetailBuilder.EditDefaultProperty(CommandLinePropertyHandle)->CustomWidget()
+	.ShouldAutoExpand(true)
+	.NameContent()	
+	[
+		CommandLinePropertyHandle->CreatePropertyNameWidget()
+	]
+	.ValueContent()
+	.HAlign(HAlign_Fill)
+	[
+		SNew(SMultiLineEditableTextBox)
+		.AllowMultiLine(false)
+		.AutoWrapText(true)
+		.Text(this, &FLauncherBasicLaunchProfileSettingsDetails::GetCommandLine)
+		.OnTextCommitted(this, &FLauncherBasicLaunchProfileSettingsDetails::SetCommandLine)
+	];
+
+	// put the device selection list at the bottom of the properties, full width
+	TSharedRef<IPropertyHandle> DeployDeviceIDPropertyHandle = DetailBuilder.GetProperty(TEXT("DeployDeviceID"));
+	DetailBuilder.EditDefaultProperty(DeployDeviceIDPropertyHandle)->CustomWidget()
+	.OverrideResetToDefault(FResetToDefaultOverride::Hide())
+	.NameContent()
+	.VAlign(VAlign_Top)
+	[
+		DeployDeviceIDPropertyHandle->CreatePropertyNameWidget()
+	]
+	.ValueContent()
+	.HAlign(HAlign_Fill)
+	[
+		SNew(SCustomLaunchDeviceListView)
+		.OnSelectionChanged(this, &FLauncherBasicLaunchProfileSettingsDetails::SetDeployDeviceIDs)
+		.SelectedDevices(this, &FLauncherBasicLaunchProfileSettingsDetails::GetDeployDeviceIDs)
+		.AllPlatforms(true)
+		.SingleSelect(true)
+	];
+}
+
+
+
+
+
+ULauncherBasicLaunchProfileSettings* FLauncherBasicLaunchProfileSettingsDetails::GetProfileSettings() const
+{
+	ULauncherBasicLaunchProfileSettings* ProfileSettings = nullptr;
+
+	const TArray<TWeakObjectPtr<UObject>>& SelectedObjects = PropertyUtilities->GetSelectedObjects();
+	for (const TWeakObjectPtr<UObject>& SelectedObject : SelectedObjects)
+	{
+		ProfileSettings = Cast<ULauncherBasicLaunchProfileSettings>(SelectedObject.Get());
+		if (ProfileSettings != nullptr)
+		{
+			break;
+		}
+	}
+
+	return ProfileSettings;
+}
+
+FString FLauncherBasicLaunchProfileSettingsDetails::GetProjectPath() const
+{
+	ULauncherBasicLaunchProfileSettings* ProfileSettings = GetProfileSettings();
+	if (ProfileSettings != nullptr && ProfileSettings->bHasProjectSpecified)
+	{
+		return ProfileSettings->ProjectPath;
+	}
+	else
+	{
+		return GetProfileManager()->GetProjectPath();
+	}
+}
+
+void FLauncherBasicLaunchProfileSettingsDetails::SetProjectName(FString ProjectPath)
+{
+	ULauncherBasicLaunchProfileSettings* ProfileSettings = GetProfileSettings();
+
+	if (ProfileSettings != nullptr)
+	{
+		ProfileSettings->ProjectPath = ProjectPath;
+
+		NotifyPropertyChanged(GET_MEMBER_NAME_CHECKED(ULauncherBasicLaunchProfileSettings, ProjectPath), ProfileSettings);
+	}
+}
+
+bool FLauncherBasicLaunchProfileSettingsDetails::HasProject() const
+{
+	ULauncherBasicLaunchProfileSettings* ProfileSettings = GetProfileSettings();
+	return (ProfileSettings != nullptr && ProfileSettings->bHasProjectSpecified);
+}
+
+
+
+FString FLauncherBasicLaunchProfileSettingsDetails::GetBuildTarget() const
+{
+	ULauncherBasicLaunchProfileSettings* ProfileSettings = GetProfileSettings();
+	if (ProfileSettings != nullptr && ProfileSettings->bHasTargetSpecified)
+	{
+		return ProfileSettings->Target;
+	}
+	else
+	{
+		return GetProfileManager()->GetBuildTarget();
+	}
+}
+
+void FLauncherBasicLaunchProfileSettingsDetails::SetBuildTarget(FString BuildTarget)
+{
+	ULauncherBasicLaunchProfileSettings* ProfileSettings = GetProfileSettings();
+
+	if (ProfileSettings != nullptr)
+	{
+		ProfileSettings->Target = BuildTarget;
+
+		NotifyPropertyChanged(GET_MEMBER_NAME_CHECKED(ULauncherBasicLaunchProfileSettings, Target), ProfileSettings);
+	}
+}
+
+CustomLaunchUI::EContentScheme FLauncherBasicLaunchProfileSettingsDetails::GetContentScheme() const
+{
+	ULauncherBasicLaunchProfileSettings* ProfileSettings = GetProfileSettings();
+	if (ProfileSettings != nullptr)
+	{
+		return (CustomLaunchUI::EContentScheme)ProfileSettings->ContentScheme;
+	}
+
+	return CustomLaunchUI::EContentScheme::PakFiles;
+}
+
+void FLauncherBasicLaunchProfileSettingsDetails::SetContentScheme(CustomLaunchUI::EContentScheme ContentScheme)
+{
+	ULauncherBasicLaunchProfileSettings* ProfileSettings = GetProfileSettings();
+	if (ProfileSettings != nullptr)
+	{
+		ProfileSettings->ContentScheme = (uint8)ContentScheme;
+		NotifyPropertyChanged(GET_MEMBER_NAME_CHECKED(ULauncherBasicLaunchProfileSettings, ContentScheme), ProfileSettings);
+	}
+
+}
+
+bool FLauncherBasicLaunchProfileSettingsDetails::IsContentSchemeAvailable(CustomLaunchUI::EContentScheme ContentScheme) const
+{
+	// @fixme: check the project supports the content scheme etc
+	return ContentScheme != CustomLaunchUI::EContentScheme::ZenPakStreaming;
+}
+
+
+
+
+FText FLauncherBasicLaunchProfileSettingsDetails::GetCommandLine() const
+{
+	ULauncherBasicLaunchProfileSettings* ProfileSettings = GetProfileSettings();
+	if (ProfileSettings != nullptr)
+	{
+		return FText::FromString(ProfileSettings->CommandLine);
+	}
+
+	return FText::GetEmpty();
+}
+
+void FLauncherBasicLaunchProfileSettingsDetails::SetCommandLine( const FText& NewText, ETextCommit::Type InTextCommit )
+{
+	ULauncherBasicLaunchProfileSettings* ProfileSettings = GetProfileSettings();
+	if (ProfileSettings != nullptr)
+	{
+		ProfileSettings->CommandLine = NewText.ToString();
+		NotifyPropertyChanged(GET_MEMBER_NAME_CHECKED(ULauncherBasicLaunchProfileSettings, CommandLine), ProfileSettings);
+
+		PropertyUtilities->RequestForceRefresh();
+	}
+}
+
+
+void FLauncherBasicLaunchProfileSettingsDetails::SetDeployDeviceIDs(TArray<FString> DeployDeviceIDs)
+{
+	ULauncherBasicLaunchProfileSettings* ProfileSettings = GetProfileSettings();
+	if (ProfileSettings != nullptr)
+	{
+		if (DeployDeviceIDs.Num() > 0)
+		{
+			ProfileSettings->DeployDeviceID = DeployDeviceIDs[0];
+		}
+		else
+		{
+			ProfileSettings->DeployDeviceID = FString();
+		}
+
+		NotifyPropertyChanged(GET_MEMBER_NAME_CHECKED(ULauncherBasicLaunchProfileSettings, DeployDeviceID), ProfileSettings);
+	}
+}
+
+TArray<FString> FLauncherBasicLaunchProfileSettingsDetails::GetDeployDeviceIDs() const
+{
+	TArray<FString> Result;
+	
+	ULauncherBasicLaunchProfileSettings* ProfileSettings = GetProfileSettings();
+	if (ProfileSettings != nullptr && !ProfileSettings->DeployDeviceID.IsEmpty())
+	{
+		Result.Add(ProfileSettings->DeployDeviceID);
+	}
+
+	return MoveTemp(Result);
+}
+
+
+
+
+
+
+
+
+
+void FLauncherCustomProfileSettingsDetails::CustomizeDetails(IDetailLayoutBuilder& DetailBuilder)
+{
+	FBaseLauncherProfileDetails::CustomizeDetails(DetailBuilder);
+
+	// customize platform selector
+	TSharedRef<IPropertyHandle> PlatformPropertyHandle = DetailBuilder.GetProperty(TEXT("Platforms"));
+	DetailBuilder.EditDefaultProperty(PlatformPropertyHandle)->ShowPropertyButtons(false)
+	.OverrideResetToDefault(FResetToDefaultOverride::Hide())
+	.CustomWidget()
+	.NameContent()
+	[
+		PlatformPropertyHandle->CreatePropertyNameWidget()
+	]
+	.ValueContent()	
+	.HAlign(HAlign_Fill)
+	[
+		SNew(SCustomLaunchPlatformCombo)
+		.SelectedPlatforms(this, &FLauncherCustomProfileSettingsDetails::GetSelectedPlatforms)
+		.OnSelectionChanged(this, &FLauncherCustomProfileSettingsDetails::SetSelectedPlatforms)
+		.BasicPlatformsOnly(CustomLaunchUI::bUseFriendlyBuildTargetSelection)
+	];
+
+
+	// customize project selector
+	TSharedRef<IPropertyHandle> ProjectsPropertyHandle = DetailBuilder.GetProperty(TEXT("ProjectPath"));
+	DetailBuilder.EditDefaultProperty(ProjectsPropertyHandle)->CustomWidget()
+	.OverrideResetToDefault(FResetToDefaultOverride::Hide())
+	.NameContent()	
+	[
+		ProjectsPropertyHandle->CreatePropertyNameWidget()
+	]
+	.ValueContent()
+	.HAlign(HAlign_Fill)
+	[
+		SNew(SCustomLaunchProjectCombo)
+		.OnSelectionChanged(this, &FLauncherCustomProfileSettingsDetails::SetProjectName)
+		.SelectedProject(this, &FLauncherCustomProfileSettingsDetails::GetProjectPath)
+		.HasProject(this, &FLauncherCustomProfileSettingsDetails::HasProject)
+		.ShowAnyProjectOption(true)
+		.CurrentProjectOption(SCustomLaunchProjectCombo::ECurrentProjectOption::ActualProject)
+		.Font(IDetailLayoutBuilder::GetDetailFont())
+	];
+	
+	// customize the build target combo
+	TSharedRef<IPropertyHandle> TargetPropertyHandle = DetailBuilder.GetProperty(TEXT("Target"));
+	DetailBuilder.EditDefaultProperty(TargetPropertyHandle)->CustomWidget()
+	.OverrideResetToDefault(FResetToDefaultOverride::Hide())
+	.NameContent()	
+	[
+		TargetPropertyHandle->CreatePropertyNameWidget()
+	]
+	.ValueContent()
+	.HAlign(HAlign_Fill)
+	[
+		SNew(SCustomLaunchBuildTargetCombo)
+		.OnSelectionChanged(this, &FLauncherCustomProfileSettingsDetails::SetBuildTarget)
+		.SelectedBuildTarget(this, &FLauncherCustomProfileSettingsDetails::GetBuildTarget)
+		.SupportedTargetTypes(this, &FLauncherCustomProfileSettingsDetails::GetSupportedBuildTargetTypes)
+		.SelectedProject(this, &FLauncherCustomProfileSettingsDetails::GetProjectPath)
+		.Font(IDetailLayoutBuilder::GetDetailFont())
+	];
+
+	// customize the configuration combo
+	TSharedPtr<IPropertyHandle> ConfigurationPropertyHandle = DetailBuilder.GetProperty(TEXT("Configuration"));
+	DetailBuilder.EditDefaultProperty(ConfigurationPropertyHandle)->CustomWidget()
+	.NameContent()
+	[
+		ConfigurationPropertyHandle->CreatePropertyNameWidget()
+	]
+	.ValueContent()
+	.HAlign(HAlign_Fill)
+	[
+		ConfigurationPropertyHandle->CreatePropertyValueWidget()
+	];
+
+
+	// customize the content scheme combo
+	TSharedRef<IPropertyHandle> ContentSchemePropertyHandle = DetailBuilder.GetProperty(TEXT("ContentScheme"));
+	DetailBuilder.EditDefaultProperty(ContentSchemePropertyHandle)->CustomWidget()
+	.NameContent()
+	[
+		ContentSchemePropertyHandle->CreatePropertyNameWidget()
+	]
+	.ValueContent()
+	.HAlign(HAlign_Fill)
+	[
+		SNew(SCustomLaunchContentSchemeCombo)
+		.OnSelectionChanged(this, &FLauncherCustomProfileSettingsDetails::SetContentScheme)
+		.SelectedContentScheme(this, &FLauncherCustomProfileSettingsDetails::GetContentScheme)
+		.IsContentSchemeAvailable(this, &FLauncherCustomProfileSettingsDetails::IsContentSchemeAvailable)
+		.Font(IDetailLayoutBuilder::GetDetailFont())
+	];
+
+
+	// customize the command line textbox
+	TSharedRef<IPropertyHandle> CommandLinePropertyHandle = DetailBuilder.GetProperty(TEXT("CommandLine"));
+	DetailBuilder.EditDefaultProperty(CommandLinePropertyHandle)->CustomWidget()
+	.NameContent()	
+	[
+		CommandLinePropertyHandle->CreatePropertyNameWidget()
+	]
+	.ValueContent()
+	.HAlign(HAlign_Fill)
+	[
+		SNew(SMultiLineEditableTextBox)
+		.AllowMultiLine(false)
+		.AutoWrapText(true)
+		.Text(this, &FLauncherCustomProfileSettingsDetails::GetCommandLine)
+		.OnTextCommitted(this, &FLauncherCustomProfileSettingsDetails::SetCommandLine)
+	];
+
+
+
+
+	// customize the additional cooker options textbox
+	TSharedRef<IPropertyHandle> AdditionalCookerOptionsPropertyHandle = DetailBuilder.GetProperty(TEXT("AdditionalCookerOptions"));
+	DetailBuilder.EditDefaultProperty(AdditionalCookerOptionsPropertyHandle)->CustomWidget()
+	.NameContent()	
+	[
+		AdditionalCookerOptionsPropertyHandle->CreatePropertyNameWidget()
+	]
+	.ValueContent()
+	.HAlign(HAlign_Fill)
+	[
+		SNew(SMultiLineEditableTextBox)
+		.AllowMultiLine(false)
+		.AutoWrapText(true)
+		.Text(this, &FLauncherCustomProfileSettingsDetails::GetAdditionalCookerOptions)
+		.OnTextCommitted(this, &FLauncherCustomProfileSettingsDetails::SetAdditionalCookerOptions)
+	];
+
+
+	// customize map selector
+	MapListView = 
+		SNew(SCustomLaunchMapListView)
+		.OnSelectionChanged(this, &FLauncherCustomProfileSettingsDetails::SetMapsToCook)
+		.SelectedMaps(this, &FLauncherCustomProfileSettingsDetails::GetMapsToCook)
+		.ProjectPath(this, &FLauncherCustomProfileSettingsDetails::GetProjectPath)
+	;
+	TSharedRef<IPropertyHandle> MapsToCookPropertyHandle = DetailBuilder.GetProperty(TEXT("MapsToCook"));
+	DetailBuilder.EditDefaultProperty(MapsToCookPropertyHandle)->ShowPropertyButtons(false)
+	.CustomWidget()
+	.NameContent()
+	.VAlign(VAlign_Top)
+	[
+		MapsToCookPropertyHandle->CreatePropertyNameWidget()
+	]
+	.ValueContent()
+	[
+		SNew(SVerticalBox)
+
+		+SVerticalBox::Slot()
+		.AutoHeight()
+		.Padding(0,2)
+		[
+			CreateMapListViewOptionsWidget()
+		]
+
+		+SVerticalBox::Slot()
+		.FillHeight(1)
+		[
+			SNew(SVerticalResizeBox)
+			.Visibility_Lambda([this](){ return (GetMapOption() == ULauncherCustomProfileSettings::EMapOption::Selected) ? EVisibility::Visible : EVisibility::Collapsed; } )
+			.HandleHeight(4)
+			.ContentHeight(this, &FLauncherCustomProfileSettingsDetails::GetMapListHeight)
+			.ContentHeightChanged(this, &FLauncherCustomProfileSettingsDetails::SetMapListHeight)
+			.HandleColor( FAppStyle::Get().GetSlateColor("Colors.Secondary").GetSpecifiedColor() )
+			[
+				MapListView.ToSharedRef()
+			]
+		]
+	];
+
+
+	// customize device selector
+	TSharedRef<IPropertyHandle> DeployDeviceIDsPropertyHandle = DetailBuilder.GetProperty(TEXT("DeployDeviceIDs"));
+	DetailBuilder.EditDefaultProperty(DeployDeviceIDsPropertyHandle)->ShowPropertyButtons(false)
+	.CustomWidget()
+	.NameContent()
+	.VAlign(VAlign_Top)
+	[
+		DeployDeviceIDsPropertyHandle->CreatePropertyNameWidget()
+	]
+	.ValueContent()
+	[
+		SNew(SVerticalBox)
+
+		+SVerticalBox::Slot()
+		.AutoHeight()
+		.Padding(0,2)
+		[
+			CreateDeployDeviceSelectorOptionsWidget()
+		]
+
+		+SVerticalBox::Slot()
+		.FillHeight(1)
+		[
+			SNew(SVerticalResizeBox)
+			.Visibility_Lambda([this](){ return (GetDeployDeviceOption() == ULauncherCustomProfileSettings::EDeployDeviceOption::Selected) ? EVisibility::Visible : EVisibility::Collapsed; } )
+			.HandleHeight(4)
+			.ContentHeight(this, &FLauncherCustomProfileSettingsDetails::GetDeployDeviceListHeight)
+			.ContentHeightChanged(this, &FLauncherCustomProfileSettingsDetails::SetDeployDeviceListHeight)
+			.HandleColor( FAppStyle::Get().GetSlateColor("Colors.Secondary").GetSpecifiedColor() )
+			[
+				SAssignNew(DeployDeviceListView, SCustomLaunchDeviceListView)
+				.OnDeviceRemoved( this, &FLauncherCustomProfileSettingsDetails::OnDeviceRemoved)
+				.OnSelectionChanged(this, &FLauncherCustomProfileSettingsDetails::SetDeployDeviceIDs)
+				.SelectedDevices(this, &FLauncherCustomProfileSettingsDetails::GetDeployDeviceIDs)
+				.Platforms(this, &FLauncherCustomProfileSettingsDetails::GetSelectedPlatforms)
+			]
+		]
+	];
+				
+	// customize zen pak streaming path field - hide it if the ZenPakStreaming content mode isn't selected
+	// this is currently necessary because EditCondition with EditConditionHides does not seem to work with FDirectoryPath
+	TSharedRef<IPropertyHandle> ZenPakStreamingPathPropertyHandle = DetailBuilder.GetProperty(TEXT("ZenPakStreamingPath"));
+	DetailBuilder.EditDefaultProperty(ZenPakStreamingPathPropertyHandle)
+		->Visibility( TAttribute<EVisibility>::CreateSP( this, &FLauncherCustomProfileSettingsDetails::GetZenPakStreamingPropertyVisibility ));
+	;
+}
+
+ULauncherCustomProfileSettings* FLauncherCustomProfileSettingsDetails::GetProfileSettings() const
+{
+	ULauncherCustomProfileSettings* ProfileSettings = nullptr;
+
+	const TArray<TWeakObjectPtr<UObject>>& SelectedObjects = PropertyUtilities->GetSelectedObjects();
+	for (const TWeakObjectPtr<UObject>& SelectedObject : SelectedObjects)
+	{
+		ProfileSettings = Cast<ULauncherCustomProfileSettings>(SelectedObject.Get());
+		if (ProfileSettings != nullptr)
+		{
+			break;
+		}
+	}
+
+	return ProfileSettings;
+}
+
+void FLauncherCustomProfileSettingsDetails::SetSelectedPlatforms( TArray<FString> SelectedPlatforms )
+{
+	ULauncherCustomProfileSettings* ProfileSettings = GetProfileSettings();
+
+	if (ProfileSettings != nullptr)
+	{
+		if (CustomLaunchUI::bUseFriendlyBuildTargetSelection)
+		{
+			FTargetInfo BuildTargetInfo = CustomLaunchUI::FModel::GetBuildTargetInfo(GetBuildTarget(), GetProjectPath());
+
+			ProfileSettings->Platforms.Reset();
+			for (const FString& Platform : SelectedPlatforms)
+			{
+				ProfileSettings->Platforms.Add( CustomLaunchUI::FModel::GetBuildTargetPlatformName(Platform, BuildTargetInfo));
+			}
+		}
+		else
+		{
+			ProfileSettings->Platforms = SelectedPlatforms;
+		}
+		NotifyPropertyChanged(GET_MEMBER_NAME_CHECKED(ULauncherCustomProfileSettings, Platforms), ProfileSettings);
+	}
+
+	DeployDeviceListView->OnSelectedPlatformChanged();
+}
+
+TArray<FString> FLauncherCustomProfileSettingsDetails::GetSelectedPlatforms() const
+{
+	TArray<FString> Platforms;
+
+	ULauncherCustomProfileSettings* ProfileSettings = GetProfileSettings();
+	if (ProfileSettings != nullptr)
+	{
+		if (CustomLaunchUI::bUseFriendlyBuildTargetSelection)
+		{
+			for (const FString& Platform : ProfileSettings->Platforms)
+			{
+				Platforms.Add(CustomLaunchUI::FModel::GetVanillaPlatformName(Platform));
+			}
+		}
+		else
+		{
+			Platforms = ProfileSettings->Platforms;
+		}
+	}
+
+	return MoveTemp(Platforms);
+}
+
+
+FString FLauncherCustomProfileSettingsDetails::GetProjectPath() const
+{
+	ULauncherCustomProfileSettings* ProfileSettings = GetProfileSettings();
+	if (ProfileSettings != nullptr && ProfileSettings->bHasProjectSpecified)
+	{
+		return ProfileSettings->ProjectPath;
+	}
+
+	return FString();
+}
+
+void FLauncherCustomProfileSettingsDetails::SetProjectName(FString ProjectPath)
+{
+	ULauncherCustomProfileSettings* ProfileSettings = GetProfileSettings();
+
+	if (ProfileSettings != nullptr)
+	{
+		ProfileSettings->bHasProjectSpecified = !ProjectPath.IsEmpty();
+		ProfileSettings->ProjectPath = ProjectPath;
+
+		NotifyPropertyChanged(GET_MEMBER_NAME_CHECKED(ULauncherCustomProfileSettings, ProjectPath), ProfileSettings);
+	}
+
+	if (MapListView.IsValid())
+	{
+		MapListView->RefreshMapList();
+	}
+}
+
+bool FLauncherCustomProfileSettingsDetails::HasProject() const
+{
+	ULauncherCustomProfileSettings* ProfileSettings = GetProfileSettings();
+	return (ProfileSettings != nullptr && ProfileSettings->bHasProjectSpecified);
+}
+
+
+FString FLauncherCustomProfileSettingsDetails::GetBuildTarget() const
+{
+	ULauncherCustomProfileSettings* ProfileSettings = GetProfileSettings();
+	if (ProfileSettings != nullptr)
+	{
+		return ProfileSettings->Target;
+	}
+	else
+	{
+		return FString();
+	}
+}
+
+void FLauncherCustomProfileSettingsDetails::SetBuildTarget(FString BuildTarget)
+{
+	ULauncherCustomProfileSettings* ProfileSettings = GetProfileSettings();
+
+	if (ProfileSettings != nullptr)
+	{
+		ProfileSettings->bHasTargetSpecified = !BuildTarget.IsEmpty();
+		ProfileSettings->Target = BuildTarget;
+
+		NotifyPropertyChanged(GET_MEMBER_NAME_CHECKED(ULauncherCustomProfileSettings, Target), ProfileSettings);
+	}
+}
+
+TArray<EBuildTargetType> FLauncherCustomProfileSettingsDetails::GetSupportedBuildTargetTypes() const
+{
+	TArray<EBuildTargetType> Result;
+
+	ULauncherCustomProfileSettings* ProfileSettings = GetProfileSettings();
+	if (ProfileSettings != nullptr && ProfileSettings->Platforms.Num() > 0)
+	{
+		const PlatformInfo::FTargetPlatformInfo* PlatformInfo = PlatformInfo::FindPlatformInfo(FName(ProfileSettings->Platforms[0]));
+		if (PlatformInfo)
+		{
+			if (CustomLaunchUI::bUseFriendlyBuildTargetSelection)
+			{
+				Result.Add(PlatformInfo->VanillaInfo->PlatformType);
+
+				for (const PlatformInfo::FTargetPlatformInfo* PlatformFlavorInfo : PlatformInfo->VanillaInfo->Flavors)
+				{
+					Result.AddUnique(PlatformFlavorInfo->PlatformType);
+				}
+			}
+			else
+			{
+				Result.Add(PlatformInfo->PlatformType);
+			}
+		}
+	}
+
+	return MoveTemp(Result);
+}
+
+CustomLaunchUI::EContentScheme FLauncherCustomProfileSettingsDetails::GetContentScheme() const
+{
+	ULauncherCustomProfileSettings* ProfileSettings = GetProfileSettings();
+	if (ProfileSettings != nullptr)
+	{
+		return (CustomLaunchUI::EContentScheme)ProfileSettings->ContentScheme;
+	}
+
+	return CustomLaunchUI::EContentScheme::PakFiles;
+}
+
+void FLauncherCustomProfileSettingsDetails::SetContentScheme(CustomLaunchUI::EContentScheme ContentScheme)
+{
+	ULauncherCustomProfileSettings* ProfileSettings = GetProfileSettings();
+	if (ProfileSettings != nullptr)
+	{
+		ProfileSettings->ContentScheme = (uint8)ContentScheme;
+		NotifyPropertyChanged(GET_MEMBER_NAME_CHECKED(ULauncherCustomProfileSettings, ContentScheme), ProfileSettings);
+	}
+
+}
+
+bool FLauncherCustomProfileSettingsDetails::IsContentSchemeAvailable(CustomLaunchUI::EContentScheme ContentScheme) const
+{
+	// @fixme: check the project supports the content scheme etc
+	return true;
+}
+
+
+
+FText FLauncherCustomProfileSettingsDetails::GetCommandLine() const
+{
+	ULauncherCustomProfileSettings* ProfileSettings = GetProfileSettings();
+	if (ProfileSettings != nullptr)
+	{
+		return FText::FromString(ProfileSettings->CommandLine);
+	}
+
+	return FText::GetEmpty();
+}
+
+void FLauncherCustomProfileSettingsDetails::SetCommandLine( const FText& NewText, ETextCommit::Type InTextCommit )
+{
+	ULauncherCustomProfileSettings* ProfileSettings = GetProfileSettings();
+	if (ProfileSettings != nullptr)
+	{
+		ProfileSettings->CommandLine = NewText.ToString();
+		NotifyPropertyChanged(GET_MEMBER_NAME_CHECKED(ULauncherCustomProfileSettings, CommandLine), ProfileSettings);
+
+		PropertyUtilities->RequestForceRefresh();
+	}
+}
+
+
+FText FLauncherCustomProfileSettingsDetails::GetAdditionalCookerOptions() const
+{
+	ULauncherCustomProfileSettings* ProfileSettings = GetProfileSettings();
+	if (ProfileSettings != nullptr)
+	{
+		return FText::FromString(ProfileSettings->AdditionalCookerOptions);
+	}
+
+	return FText::GetEmpty();
+}
+
+void FLauncherCustomProfileSettingsDetails::SetAdditionalCookerOptions( const FText& NewText, ETextCommit::Type InTextCommit )
+{
+	ULauncherCustomProfileSettings* ProfileSettings = GetProfileSettings();
+	if (ProfileSettings != nullptr)
+	{
+		ProfileSettings->AdditionalCookerOptions = NewText.ToString();
+		NotifyPropertyChanged(GET_MEMBER_NAME_CHECKED(ULauncherCustomProfileSettings, AdditionalCookerOptions), ProfileSettings);
+
+		PropertyUtilities->RequestForceRefresh();
+	}
+}
+
+
+BEGIN_SLATE_FUNCTION_BUILD_OPTIMIZATION
+TSharedRef<SWidget> FLauncherCustomProfileSettingsDetails::CreateMapListViewOptionsWidget()
+{
+	return SNew(SHorizontalBox)
+
+		+SHorizontalBox::Slot()
+		.Padding(0, 0)
+		.AutoWidth()
+		.VAlign(VAlign_Center)
+		[
+			SNew(SSegmentedControl<ULauncherCustomProfileSettings::EMapOption>)
+			.Value( this, &FLauncherCustomProfileSettingsDetails::GetMapOption)
+			.OnValueChanged(this, &FLauncherCustomProfileSettingsDetails::SetMapOption)
+
+			+SSegmentedControl<ULauncherCustomProfileSettings::EMapOption>::Slot(ULauncherCustomProfileSettings::EMapOption::Startup)
+			.HAlign(HAlign_Center)
+			.VAlign(VAlign_Center)
+			[
+				SNew(STextBlock)
+				.Text(LOCTEXT("StartupMapsLabel", "Startup Maps"))
+				.Font(IDetailLayoutBuilder::GetDetailFont())
+			]
+
+			+SSegmentedControl<ULauncherCustomProfileSettings::EMapOption>::Slot(ULauncherCustomProfileSettings::EMapOption::Selected)
+			.HAlign(HAlign_Center)
+			.VAlign(VAlign_Center)
+			[
+				SNew(STextBlock)
+				.Text(LOCTEXT("SelectedMapsLabel", "Selected Maps"))
+				.Font(IDetailLayoutBuilder::GetDetailFont())
+			]
+		]
+
+		// map selector controls
+		+SHorizontalBox::Slot()
+		.Padding(8, 0)
+		.FillWidth(1)
+		.VAlign(VAlign_Center)
+		[
+			SNew(SBox)
+			.Visibility_Lambda([this](){ return GetMapOption() == ULauncherCustomProfileSettings::EMapOption::Selected  ? EVisibility::Visible : EVisibility::Collapsed; } )
+			[
+				MapListView->MakeControlsWidget()
+			]
+		]
+	;
+}
+END_SLATE_FUNCTION_BUILD_OPTIMIZATION
+
+void FLauncherCustomProfileSettingsDetails::SetMapsToCook(TArray<FString> MapsToCook)
+{
+	ULauncherCustomProfileSettings* ProfileSettings = GetProfileSettings();
+	if (ProfileSettings != nullptr)
+	{
+		ProfileSettings->MapsToCook = MapsToCook;
+		NotifyPropertyChanged(GET_MEMBER_NAME_CHECKED(ULauncherCustomProfileSettings, MapsToCook), ProfileSettings);
+	}
+}
+
+TArray<FString> FLauncherCustomProfileSettingsDetails::GetMapsToCook() const
+{
+	TArray<FString> MapsToCook;
+	ULauncherCustomProfileSettings* ProfileSettings = GetProfileSettings();
+	if (ProfileSettings != nullptr)
+	{
+		MapsToCook = ProfileSettings->MapsToCook;
+	}
+
+	return MoveTemp(MapsToCook);
+}
+
+float FLauncherCustomProfileSettingsDetails::GetMapListHeight() const
+{
+	return MapListHeight;
+}
+
+void FLauncherCustomProfileSettingsDetails::SetMapListHeight( float NewHeight )
+{
+	static const float MinMapListHeight = 100.0f;
+
+	MapListHeight = FMath::Max(NewHeight, MinMapListHeight);
+}
+
+ULauncherCustomProfileSettings::EMapOption FLauncherCustomProfileSettingsDetails::GetMapOption() const
+{
+	ULauncherCustomProfileSettings* ProfileSettings = GetProfileSettings();
+	if (ProfileSettings != nullptr)
+	{
+		return ProfileSettings->MapOption;
+	}
+	else
+	{
+		return ULauncherCustomProfileSettings::EMapOption::Startup;
+	}
+}
+
+void FLauncherCustomProfileSettingsDetails::SetMapOption( ULauncherCustomProfileSettings::EMapOption NewMapOption )
+{
+	bool bShow = (NewMapOption == ULauncherCustomProfileSettings::EMapOption::Selected);
+
+	ULauncherCustomProfileSettings* ProfileSettings = GetProfileSettings();
+	if (ProfileSettings != nullptr)
+	{
+		ProfileSettings->MapOption = NewMapOption;
+
+		if (bShow)
+		{
+			// restore the cooked maps again, if any
+			if (ProfileSettings->CachedMapsToCook.Num() > 0 && ProfileSettings->MapsToCook.Num() == 0)
+			{
+				SetMapsToCook(ProfileSettings->CachedMapsToCook);
+				ProfileSettings->CachedMapsToCook.Reset();
+			}
+		}
+		else
+		{
+			// to set the 'cook startup maps only', its necessary to remove all the cooked maps - take a copy of the values to allow it to be restored
+			ProfileSettings->CachedMapsToCook = ProfileSettings->MapsToCook;
+			SetMapsToCook(TArray<FString>());
+		}
+	}
+
+	MapListView->RefreshMapList();
+}
+
+
+
+
+BEGIN_SLATE_FUNCTION_BUILD_OPTIMIZATION
+TSharedRef<SWidget> FLauncherCustomProfileSettingsDetails::CreateDeployDeviceSelectorOptionsWidget()
+{
+	return SNew(SHorizontalBox)
+
+		+SHorizontalBox::Slot()
+		.Padding(0, 0)
+		.AutoWidth()
+		.VAlign(VAlign_Center)
+		[
+			SNew(SSegmentedControl<ULauncherCustomProfileSettings::EDeployDeviceOption>)
+			.Value( this, &FLauncherCustomProfileSettingsDetails::GetDeployDeviceOption)
+			.OnValueChanged(this, &FLauncherCustomProfileSettingsDetails::SetDeployDeviceOption)
+
+			+SSegmentedControl<ULauncherCustomProfileSettings::EDeployDeviceOption>::Slot(ULauncherCustomProfileSettings::EDeployDeviceOption::Default)
+			.HAlign(HAlign_Center)
+			.VAlign(VAlign_Center)
+			[
+				SNew(STextBlock)
+				.Text(LOCTEXT("DefaultDeviceLabel", "Default Device"))
+				.Font(IDetailLayoutBuilder::GetDetailFont())
+			]
+
+			+SSegmentedControl<ULauncherCustomProfileSettings::EDeployDeviceOption>::Slot(ULauncherCustomProfileSettings::EDeployDeviceOption::Selected)
+			.HAlign(HAlign_Center)
+			.VAlign(VAlign_Center)
+			[
+				SNew(STextBlock)
+				.Text(LOCTEXT("SelectedDevicesLabel", "Selected Devices"))
+				.Font(IDetailLayoutBuilder::GetDetailFont())
+			]
+		]
+	;
+}
+END_SLATE_FUNCTION_BUILD_OPTIMIZATION
+
+
+void FLauncherCustomProfileSettingsDetails::SetDeployDeviceIDs(TArray<FString> DeployDeviceIDs)
+{
+	ULauncherCustomProfileSettings* ProfileSettings = GetProfileSettings();
+	if (ProfileSettings != nullptr)
+	{
+		ProfileSettings->DeployDeviceIDs = DeployDeviceIDs;
+		NotifyPropertyChanged(GET_MEMBER_NAME_CHECKED(ULauncherCustomProfileSettings, DeployDeviceIDs), ProfileSettings);
+	}
+}
+
+TArray<FString> FLauncherCustomProfileSettingsDetails::GetDeployDeviceIDs() const
+{
+	TArray<FString> DeployDeviceIDs;
+	ULauncherCustomProfileSettings* ProfileSettings = GetProfileSettings();
+	if (ProfileSettings != nullptr)
+	{
+		DeployDeviceIDs = ProfileSettings->DeployDeviceIDs;
+	}
+
+	return MoveTemp(DeployDeviceIDs);
+}
+
+float FLauncherCustomProfileSettingsDetails::GetDeployDeviceListHeight() const
+{
+	return DeployDeviceListHeight;
+}
+
+void FLauncherCustomProfileSettingsDetails::SetDeployDeviceListHeight( float NewHeight )
+{
+	static const float MinDeployDeviceListHeight = 100.0f;
+
+	DeployDeviceListHeight = FMath::Max(NewHeight, MinDeployDeviceListHeight);
+}
+
+void FLauncherCustomProfileSettingsDetails::OnDeviceRemoved(FString DeviceID)
+{
+	ULauncherCustomProfileSettings* ProfileSettings = GetProfileSettings();
+	if (ProfileSettings != nullptr)
+	{
+		ProfileSettings->CachedDeployDeviceIDs.Remove(DeviceID);
+	}
+}
+
+ULauncherCustomProfileSettings::EDeployDeviceOption FLauncherCustomProfileSettingsDetails::GetDeployDeviceOption() const
+{
+	ULauncherCustomProfileSettings* ProfileSettings = GetProfileSettings();
+	if (ProfileSettings != nullptr)
+	{
+		return ProfileSettings->DeployDeviceOption;
+	}
+	else
+	{
+		return ULauncherCustomProfileSettings::EDeployDeviceOption::Default;
+	}
+}
+
+void FLauncherCustomProfileSettingsDetails::SetDeployDeviceOption( ULauncherCustomProfileSettings::EDeployDeviceOption NewDeployDeviceOption )
+{
+	bool bShow = (NewDeployDeviceOption == ULauncherCustomProfileSettings::EDeployDeviceOption::Selected);
+
+	ULauncherCustomProfileSettings* ProfileSettings = GetProfileSettings();
+	if (ProfileSettings != nullptr)
+	{
+		ProfileSettings->DeployDeviceOption = NewDeployDeviceOption;
+
+		if (bShow)
+		{
+			// restore the deployed device list again, if any
+			if (ProfileSettings->CachedDeployDeviceIDs.Num() > 0 && ProfileSettings->DeployDeviceIDs.Num() == 0)
+			{
+				SetDeployDeviceIDs(ProfileSettings->DeployDeviceIDs);
+				ProfileSettings->CachedDeployDeviceIDs.Reset();
+			}
+		}
+		else
+		{
+			// to set the 'default' deploy option, its necessary to remove all the devices - take a copy of the values to allow it to be restored
+			ProfileSettings->CachedDeployDeviceIDs = ProfileSettings->DeployDeviceIDs;
+			SetDeployDeviceIDs(TArray<FString>());
+		}
+	}
+
+	DeployDeviceListView->RefreshDeviceList();
+}
+
+
+EVisibility FLauncherCustomProfileSettingsDetails::GetZenPakStreamingPropertyVisibility() const
+{
+	ULauncherCustomProfileSettings* ProfileSettings = GetProfileSettings();
+	if (ProfileSettings != nullptr && ProfileSettings->ContentScheme == (uint8)CustomLaunchUI::EContentScheme::ZenPakStreaming)
+	{
+		return EVisibility::Visible;
+	}
+
+	return EVisibility::Collapsed;
+}
+
+
+
+
+
+void FBaseLauncherProfileDetails::CustomizeDetails(IDetailLayoutBuilder& DetailBuilder)
+{
+	// cache property utilities
+	PropertyUtilities = DetailBuilder.GetPropertyUtilities().ToSharedPtr();
+}
+
+ILauncherProfileManagerRef FBaseLauncherProfileDetails::GetProfileManager() const
+{
+	ILauncherServicesModule& LauncherServicesModule = FModuleManager::LoadModuleChecked<ILauncherServicesModule>("LauncherServices");
+	return LauncherServicesModule.GetProfileManager();
+}
+
+void FBaseLauncherProfileDetails::NotifyPropertyChanged(FProperty* Property, UObject* OwnerObject) const
+{
+	// signal to the ULauncherXXXProfileSettings handler
+	TArray<const UObject*> NotifyTopLevelObjects;
+	NotifyTopLevelObjects.Add(OwnerObject);
+
+	FPropertyChangedEvent ChangeEvent(Property, EPropertyChangeType::ValueSet, MakeArrayView(NotifyTopLevelObjects));
+	OwnerObject->PostEditChangeProperty(ChangeEvent);
+}
+
+void FBaseLauncherProfileDetails::NotifyPropertyChanged(FName PropertyName, UObject* OwnerObject) const
+{
+	FProperty* Property = FindFProperty<FProperty>( OwnerObject->GetClass(), PropertyName );
+	if (Property != nullptr)
+	{
+		NotifyPropertyChanged(Property, OwnerObject);
+	}
+}
+
+
+
+
+#undef LOCTEXT_NAMESPACE
diff --git a/Engine/Plugins/Developer/CustomLaunchUI/Source/Private/Details/LauncherProfileSettingsDetails.h b/Engine/Plugins/Developer/CustomLaunchUI/Source/Private/Details/LauncherProfileSettingsDetails.h
new file mode 100644
--- /dev/null
+++ b/Engine/Plugins/Developer/CustomLaunchUI/Source/Private/Details/LauncherProfileSettingsDetails.h
@@ -0,0 +1,129 @@
+// Copyright Epic Games, Inc. All Rights Reserved.
+
+#pragma once
+#include "PropertyHandle.h"
+#include "IDetailCustomization.h"
+#include "IPropertyTypeCustomization.h"
+#include "IPropertyUtilities.h"
+#include "ILauncherProfileManager.h"
+#include "Classes/LauncherProfileSettings.h"
+
+class SCustomLaunchMapListView;
+class SCustomLaunchDeviceListView;
+namespace CustomLaunchUI
+{
+	enum class EContentScheme : uint8;
+}
+
+
+class FBaseLauncherProfileDetails : public IDetailCustomization
+{
+public:
+	virtual void CustomizeDetails(IDetailLayoutBuilder& DetailBuilder) override;
+
+protected:
+	void NotifyPropertyChanged(FProperty* Property, UObject* OwnerObject) const;
+	void NotifyPropertyChanged(FName PropertyName, UObject* OwnerObject) const;
+	ILauncherProfileManagerRef GetProfileManager() const;
+
+	TSharedPtr<IPropertyUtilities> PropertyUtilities;
+};
+
+
+
+class FLauncherBasicLaunchProfileSettingsDetails : public FBaseLauncherProfileDetails
+{
+public:
+	virtual void CustomizeDetails(IDetailLayoutBuilder& DetailBuilder) override;
+
+	static TSharedRef<IDetailCustomization> MakeInstance()
+	{
+		return MakeShareable(new FLauncherBasicLaunchProfileSettingsDetails);
+	}
+
+protected:
+	class ULauncherBasicLaunchProfileSettings* GetProfileSettings() const;
+
+	FString GetProjectPath() const;
+	void SetProjectName(FString ProjectPath);
+	bool HasProject() const;
+
+	void SetBuildTarget(FString BuildTarget);
+	FString GetBuildTarget() const;
+
+	CustomLaunchUI::EContentScheme GetContentScheme() const;
+	void SetContentScheme(CustomLaunchUI::EContentScheme ContentScheme);
+	bool IsContentSchemeAvailable(CustomLaunchUI::EContentScheme) const;
+
+	FText GetCommandLine() const;
+	void SetCommandLine( const FText& NewText, ETextCommit::Type InTextCommit );
+
+
+	void SetDeployDeviceIDs(TArray<FString> DeployDeviceIDs);
+	TArray<FString> GetDeployDeviceIDs() const;
+
+
+};
+
+
+
+class FLauncherCustomProfileSettingsDetails : public FBaseLauncherProfileDetails
+{
+public:
+	virtual void CustomizeDetails(IDetailLayoutBuilder& DetailBuilder) override;
+
+	static TSharedRef<IDetailCustomization> MakeInstance()
+	{
+		return MakeShareable(new FLauncherCustomProfileSettingsDetails);
+	}
+
+protected:
+	class ULauncherCustomProfileSettings* GetProfileSettings() const;
+
+	FString GetProjectPath() const;
+	void SetProjectName(FString ProjectPath);
+	bool HasProject() const;
+
+	FString GetBuildTarget() const;
+	void SetBuildTarget(FString BuildTarget);
+	TArray<EBuildTargetType> GetSupportedBuildTargetTypes() const;
+
+	CustomLaunchUI::EContentScheme GetContentScheme() const;
+	void SetContentScheme(CustomLaunchUI::EContentScheme ContentScheme);
+	bool IsContentSchemeAvailable(CustomLaunchUI::EContentScheme) const;
+
+	FText GetCommandLine() const;
+	void SetCommandLine( const FText& NewText, ETextCommit::Type InTextCommit );
+
+
+	void SetSelectedPlatforms( TArray<FString> SelectedPlatforms );
+	TArray<FString> GetSelectedPlatforms() const;
+
+	FText GetAdditionalCookerOptions() const;
+	void SetAdditionalCookerOptions( const FText& NewText, ETextCommit::Type InTextCommit );
+
+	void SetMapsToCook(TArray<FString> MapsToCook);
+	TArray<FString> GetMapsToCook() const;
+	ULauncherCustomProfileSettings::EMapOption GetMapOption() const;
+	void SetMapOption( ULauncherCustomProfileSettings::EMapOption MapOption );
+	float GetMapListHeight() const;
+	void SetMapListHeight( float NewHeight );
+	TSharedRef<SWidget> CreateMapListViewOptionsWidget();
+	TSharedPtr<SCustomLaunchMapListView> MapListView;
+	float MapListHeight = 300.0f;
+
+
+	void SetDeployDeviceIDs(TArray<FString> DeployDeviceIDs);
+	TArray<FString> GetDeployDeviceIDs() const;
+	ULauncherCustomProfileSettings::EDeployDeviceOption GetDeployDeviceOption() const;
+	void SetDeployDeviceOption( ULauncherCustomProfileSettings::EDeployDeviceOption DeployDeviceOption );
+	float GetDeployDeviceListHeight() const;
+	void SetDeployDeviceListHeight( float NewHeight );
+	void OnDeviceRemoved(FString DeviceID);
+	TSharedRef<SWidget> CreateDeployDeviceSelectorOptionsWidget();
+	TSharedPtr<SCustomLaunchDeviceListView> DeployDeviceListView;
+	float DeployDeviceListHeight = 150.0f;
+
+	EVisibility GetZenPakStreamingPropertyVisibility() const;
+
+};
diff --git a/Engine/Plugins/Developer/CustomLaunchUI/Source/Private/Model/CustomLaunchUIModel.cpp b/Engine/Plugins/Developer/CustomLaunchUI/Source/Private/Model/CustomLaunchUIModel.cpp
new file mode 100644
--- /dev/null
+++ b/Engine/Plugins/Developer/CustomLaunchUI/Source/Private/Model/CustomLaunchUIModel.cpp
@@ -0,0 +1,652 @@
+// Copyright Epic Games, Inc. All Rights Reserved.
+
+#include "Model/CustomLaunchUIModel.h"
+#include "Misc/Paths.h"
+#include "Misc/ConfigCacheIni.h"
+#include "Misc/App.h"
+#include "PlatformInfo.h"
+#include "GameProjectHelper.h"
+#include "ITargetDeviceProxy.h"
+#include "ITargetDeviceProxyManager.h"
+#include "ITargetDeviceServicesModule.h"
+#include "Serialization/BufferArchive.h"
+#include "Serialization/MemoryReader.h"
+#include "Modules/ModuleManager.h"
+
+
+#define LOCTEXT_NAMESPACE "SCustomLaunchProfileSelector"
+
+namespace CustomLaunchUI
+{
+	// this option filters the platforms & targets based on each other. 
+	// it makes it easier to choose the desired build target but the logic under the hood is more complex
+	bool bUseFriendlyBuildTargetSelection = true;
+
+
+	FModel::FModel(const TSharedRef<ITargetDeviceProxyManager>& InDeviceProxyManager, const TSharedRef<ILauncher>& InLauncher, const TSharedRef<ILauncherProfileManager>& InProfileManager)
+		: DeviceProxyManager(InDeviceProxyManager)
+		, Launcher(InLauncher)
+		, ProfileManager(InProfileManager)
+	{
+		ProfileManager->OnProfileAdded().AddRaw(this, &FModel::HandleProfileManagerProfileAdded);
+		ProfileManager->OnProfileRemoved().AddRaw(this, &FModel::HandleProfileManagerProfileRemoved);
+
+		DeviceProxyManager->OnProxyAdded().AddRaw(this, &FModel::HandleDeviceProxyAdded);
+		DeviceProxyManager->OnProxyRemoved().AddRaw(this, &FModel::HandleDeviceProxyRemoved);
+
+		if (GIsEditor && FPaths::IsProjectFilePathSet())
+		{
+			ProfileManager->SetProjectPath(FPaths::GetProjectFilePath());
+		}
+
+		LoadConfig();
+
+
+		AllProfiles = ProfileManager->GetAllProfiles();
+
+
+		// prepare basic launch profile
+		BasicLaunchProfile = CreateBasicLaunchProfile();
+		bHasSetBasicLaunchProfilePlatform = (BasicLaunchProfile->GetCookedPlatforms().Num() > 0);
+
+		AllProfiles.Add(BasicLaunchProfile);
+		SortProfiles();
+	}
+
+
+
+	FModel::~FModel()
+	{
+		ProfileManager->OnProfileAdded().RemoveAll(this);
+		ProfileManager->OnProfileRemoved().RemoveAll(this);
+
+		DeviceProxyManager->OnProxyAdded().RemoveAll(this);
+		DeviceProxyManager->OnProxyRemoved().RemoveAll(this);
+
+		SaveConfig();
+	}
+
+
+	void FModel::SelectProfile(const ILauncherProfilePtr& NewProfile)
+	{
+		if (SelectedProfile != NewProfile)
+		{
+			ILauncherProfilePtr PreviousProfile = SelectedProfile;
+			SelectedProfile = NewProfile;
+
+			ProfileSelectedDelegate.Broadcast(SelectedProfile, PreviousProfile);
+		}
+	}
+
+
+	bool FModel::IsAdvancedProfile(const ILauncherProfileRef& Profile) const
+	{
+		return Profile->GetCookedPlatforms().Num() > 1 ||
+				Profile->GetCookedCultures().Num() > 1 ||
+				Profile->IsCreatingDLC() ||
+				Profile->IsCreatingReleaseVersion() ||
+				Profile->IsGeneratingPatch() ||				
+				Profile->GetPackagingMode() == ELauncherProfilePackagingModes::SharedRepository ||
+				Profile->GetDeploymentMode() == ELauncherProfileDeploymentModes::CopyRepository ||
+				Profile->GetLaunchMode() == ELauncherProfileLaunchModes::CustomRoles
+			;
+	}
+
+
+
+	void FModel::LoadConfig()
+	{
+		if (!GIsEditor)
+		{
+			// restore the previous project selection
+			FString ProjectPath;
+	
+			if (FPaths::IsProjectFilePathSet())
+			{
+				ProjectPath = FPaths::GetProjectFilePath();
+			}
+			else if (FGameProjectHelper::IsGameAvailable(FApp::GetProjectName()))
+			{
+				ProjectPath = FPaths::RootDir() / FApp::GetProjectName() / FApp::GetProjectName() + TEXT(".uproject");
+			}
+			else if (GConfig != nullptr)
+			{
+				GConfig->GetString(TEXT("CustomLaunchUI"), TEXT("SelectedProjectPath"), ProjectPath, GEngineIni); // @todo: keeping the same
+			}
+	
+			ProfileManager->SetProjectPath(ProjectPath);
+		}
+	}
+
+
+
+	void FModel::SaveConfig()
+	{
+		if (!GIsEditor && GConfig != nullptr && !FPaths::IsProjectFilePathSet() && !FGameProjectHelper::IsGameAvailable(FApp::GetProjectName()))
+		{
+			FString ProjectPath = ProfileManager->GetProjectPath();
+			GConfig->SetString(TEXT("CustomLaunchUI"), TEXT("SelectedProjectPath"), *ProjectPath, GEngineIni);
+		}
+	}
+
+
+
+	const PlatformInfo::FTargetPlatformInfo* FModel::GetPlatformInfo(const ILauncherProfilePtr& Profile)
+	{
+		if (Profile.IsValid() && Profile->GetCookedPlatforms().Num() > 0)
+		{
+			FString SelectedPlatform = Profile->GetCookedPlatforms()[0];
+			return PlatformInfo::FindPlatformInfo(FName(SelectedPlatform));
+		}
+
+		return nullptr;
+	}
+
+	const PlatformInfo::FTargetPlatformInfo* FModel::GetPlatformInfo(FName PlatformName, const FTargetInfo& BuildTargetInfo)
+	{
+		const PlatformInfo::FTargetPlatformInfo* PlatformInfo = PlatformInfo::FindPlatformInfo(PlatformName);
+		if (PlatformInfo == nullptr)
+		{
+			return nullptr;
+		}
+
+		// see if we found the platform immediately
+		if (BuildTargetInfo.Name.IsEmpty() || PlatformInfo->PlatformType == BuildTargetInfo.Type)
+		{
+			return PlatformInfo;
+		}
+
+		// try to find a matching flavor for the given platform & build target
+		if (PlatformInfo->VanillaInfo->PlatformType == BuildTargetInfo.Type && PlatformInfo->VanillaInfo->PlatformFlavor == PlatformInfo->PlatformFlavor)
+		{
+			return PlatformInfo->VanillaInfo;
+		}
+		for (const PlatformInfo::FTargetPlatformInfo* PlatformInfoFlavor : PlatformInfo->VanillaInfo->Flavors)
+		{
+			if (PlatformInfoFlavor->PlatformType == BuildTargetInfo.Type && PlatformInfoFlavor->PlatformFlavor == PlatformInfo->PlatformFlavor)
+			{
+				return PlatformInfoFlavor;
+			}
+		}
+
+		return nullptr;
+	}
+
+	FTargetInfo FModel::GetBuildTargetInfo( const FString& BuildTargetName, const FString& ProjectPath )
+	{
+		FTargetInfo Result;
+		if (!BuildTargetName.IsEmpty())
+		{
+			const TArray<FTargetInfo>& BuildTargets = FDesktopPlatformModule::Get()->GetTargetsForProject(ProjectPath);
+			for (const FTargetInfo& BuildTarget : BuildTargets)
+			{
+				if (BuildTarget.Name == BuildTargetName)
+				{
+					Result = BuildTarget;
+					break;
+				}
+			}
+		}
+
+		return MoveTemp(Result);
+	}
+
+	FTargetInfo FModel::GetBuildTargetInfo( const ILauncherProfileRef& Profile )
+	{
+		const FString& BuildTargetName = Profile->GetBuildTarget();
+		const FString& ProjectPath = Profile->GetProjectPath();
+		return GetBuildTargetInfo(BuildTargetName, ProjectPath);
+	}
+
+	FString FModel::GetVanillaPlatformName( const FString& PlatformName )
+	{
+		const PlatformInfo::FTargetPlatformInfo* PlatformInfo = PlatformInfo::FindPlatformInfo(*PlatformName);
+		if (PlatformInfo != nullptr)
+		{
+			return PlatformInfo->VanillaInfo->Name.ToString();
+		}
+
+		return PlatformName;
+	}
+
+	FString FModel::GetBuildTargetPlatformName( const FString& PlatformName, const FTargetInfo& BuildTargetInfo )
+	{
+		const PlatformInfo::FTargetPlatformInfo* PlatformInfo = CustomLaunchUI::FModel::GetPlatformInfo(*PlatformName, BuildTargetInfo);
+		if (PlatformInfo != nullptr)
+		{
+			return PlatformInfo->Name.ToString();
+		}
+
+		return PlatformName;
+
+	}
+
+
+
+	void FModel::SortProfiles()
+	{
+		AllProfiles.Sort( [this](const ILauncherProfilePtr& ProfileA, const ILauncherProfilePtr& ProfileB )
+		{
+			if (ProfileA == BasicLaunchProfile)
+			{
+				return true;
+			}
+			else if (ProfileB == BasicLaunchProfile)
+			{
+				return false;
+			}
+			else
+			{
+				return ProfileA->GetName() < ProfileB->GetName();
+			}
+		});
+	}
+
+	void FModel::HandleProfileManagerProfileAdded(const ILauncherProfileRef& Profile)
+	{
+		AllProfiles.Add(Profile);
+	}
+
+
+
+	void FModel::HandleProfileManagerProfileRemoved(const ILauncherProfileRef& Profile)
+	{
+		AllProfiles.Remove(Profile);
+
+		if (Profile == SelectedProfile)
+		{
+			SelectProfile(BasicLaunchProfile);
+		}
+	}
+
+
+
+
+	void FModel::HandleDeviceProxyAdded(const TSharedRef<ITargetDeviceProxy>& DeviceProxy )
+	{
+		if (!bHasSetBasicLaunchProfilePlatform)
+		{
+			UpdatedCookedPlatformsFromDeployDeviceProxy(BasicLaunchProfile.ToSharedRef(), DeviceProxy);
+			bHasSetBasicLaunchProfilePlatform = true;
+		}
+	}
+
+	void FModel::HandleDeviceProxyRemoved(const TSharedRef<ITargetDeviceProxy>& DeviceProxy )
+	{
+	}
+
+
+
+
+	TSharedPtr<ITargetDeviceProxy> FModel::GetDeviceProxy(const ILauncherProfileRef& Profile)
+	{
+		ILauncherDeviceGroupPtr DeployedDeviceGroup = Profile->GetDeployedDeviceGroup();
+		if (DeployedDeviceGroup != nullptr && DeployedDeviceGroup->GetDeviceIDs().Num() > 0)
+		{
+			const FString& DeviceID = DeployedDeviceGroup->GetDeviceIDs()[0];
+
+			ITargetDeviceServicesModule& TargetDeviceServicesModule = FModuleManager::LoadModuleChecked<ITargetDeviceServicesModule>("TargetDeviceServices");
+
+			return TargetDeviceServicesModule.GetDeviceProxyManager()->FindProxyDeviceForTargetDevice(DeviceID);
+		}
+
+		return nullptr;
+	}
+
+
+	ILauncherProfileRef FModel::CreateCustomProfile( const TCHAR* Name )
+	{
+		//create the profile
+		ILauncherProfileRef Profile = ProfileManager->CreateUnsavedProfile(Name);
+
+		// set defaults
+		SetProfileContentScheme(EContentScheme::ZenStreaming, Profile );
+		Profile->SetBuildConfiguration(EBuildConfiguration::Development);
+		Profile->SetLaunchMode(ELauncherProfileLaunchModes::DefaultRole);
+		Profile->SetBuildMode(ELauncherProfileBuildModes::Auto);
+		Profile->SetBuildUAT(false);
+		// @fixme: set all defaults here 
+		// ...
+
+
+		// make sure there is a device & deploy group
+		ILauncherDeviceGroupRef DeployDeviceGroup = ProfileManager->AddNewDeviceGroup();
+		Profile->SetDeployedDeviceGroup(DeployDeviceGroup);
+
+		TArray<TSharedPtr<ITargetDeviceProxy>> DeviceProxies;
+		DeviceProxyManager->GetProxies(NAME_None, true, DeviceProxies);
+		if (DeviceProxies.Num() > 0 && DeviceProxies[0].IsValid())
+		{
+			UpdatedCookedPlatformsFromDeployDeviceProxy(Profile, DeviceProxies[0].ToSharedRef());
+		}
+
+		return Profile;
+	}
+
+	ILauncherProfileRef FModel::CreateBasicLaunchProfile()
+	{
+		const FText BasicLaunchProfileName = LOCTEXT("BasicLaunchProfileName", "Basic Launch");
+		const FText BasicLaunchProfileDescription = LOCTEXT("BasicLaunchProfileDescription", "Use this profile to launch on a device with the recommended defaults");
+
+		ILauncherProfileRef Profile = CreateCustomProfile(*BasicLaunchProfileName.ToString());
+		Profile->SetDescription(*BasicLaunchProfileDescription.ToString());
+		Profile->SetProjectSpecified(false);
+		Profile->SetBuildTargetSpecified(false);
+
+		return Profile;
+	}
+
+
+
+	void FModel::UpdatedCookedPlatformsFromDeployDeviceProxy(const ILauncherProfileRef& Profile, TSharedPtr<ITargetDeviceProxy> DeviceProxy)
+	{
+		if (DeviceProxy.IsValid())
+		{
+			Profile->GetDeployedDeviceGroup()->RemoveAllDevices();
+			Profile->GetDeployedDeviceGroup()->AddDevice(DeviceProxy->GetTargetDeviceId(NAME_None));
+		}
+		else if (Profile->GetDeployedDeviceGroup()->GetNumDevices() > 0)
+		{
+			DeviceProxy = GetDeviceProxy(Profile);
+		}
+		
+		if (DeviceProxy.IsValid())
+		{
+			Profile->ClearCookedPlatforms();
+		
+			FTargetInfo BuildTargetInfo = GetBuildTargetInfo(Profile);
+			FString PlatformName = DeviceProxy->GetTargetPlatformName(NAME_None);
+			const PlatformInfo::FTargetPlatformInfo* PlatformInfo = GetPlatformInfo(*PlatformName, BuildTargetInfo);
+
+			if (PlatformInfo != nullptr)
+			{
+				Profile->AddCookedPlatform(PlatformInfo->Name.ToString());
+			}
+			else
+			{
+				Profile->AddCookedPlatform(DeviceProxy->GetTargetPlatformName(NAME_None));
+			}
+
+			// @todo: is default deploy platform necessary now?
+		}
+	}
+
+	void FModel::UpdateCookedPlatformsFromBuildTarget(const ILauncherProfileRef& Profile)
+	{
+		FTargetInfo BuildTargetInfo = GetBuildTargetInfo(Profile);
+
+		TArray<FString> Platforms = Profile->GetCookedPlatforms();
+		Profile->ClearCookedPlatforms();
+		for (const FString& Platform : Platforms)
+		{
+			const PlatformInfo::FTargetPlatformInfo* PlatformInfo = GetPlatformInfo(*Platform, BuildTargetInfo);
+			if (PlatformInfo != nullptr)
+			{
+				Profile->AddCookedPlatform(PlatformInfo->Name.ToString());
+			}
+			else
+			{
+				Profile->AddCookedPlatform(Platform);
+			}
+		}
+	}
+
+
+	ILauncherProfilePtr FModel::CloneCustomProfile(const ILauncherProfileRef& Profile)
+	{
+		ILauncherProfilePtr NewProfile;
+
+		FBufferArchive Writer;
+		if (Profile->Serialize(Writer))
+		{
+			NewProfile = CreateCustomProfile(TEXT("Cloned"));
+			if (NewProfile.IsValid())
+			{
+				FMemoryReader Reader(Writer);
+				NewProfile->Serialize(Reader);
+				NewProfile->AssignId(true); // need to give it a new Id - don't use the cloned one!
+
+				// @todo: add these to FLauncherProfile's Serialize() function & update the internal version. Waiting for all Zen properties to be decided (zen workspace etc)
+				NewProfile->SetUseZenStore(Profile->IsUsingZenStore());
+				// ...
+
+				ILauncherDeviceGroupRef DeployDeviceGroup = ProfileManager->AddNewDeviceGroup();
+				NewProfile->SetDeployedDeviceGroup(DeployDeviceGroup);
+				for (const FString& DeviceID : Profile->GetDeployedDeviceGroup()->GetDeviceIDs())
+				{
+					DeployDeviceGroup->AddDevice(DeviceID);
+				}
+			}
+		}
+
+		return NewProfile;
+	}
+
+			
+
+	EContentScheme FModel::DetermineProfileContentScheme(const ILauncherProfileRef& Profile) const
+	{
+		if (Profile->GetCookMode() == ELauncherProfileCookModes::OnTheFly)
+		{
+			return EContentScheme::CookOnTheFly;
+		}
+		else if (Profile->GetPackagingMode() != ELauncherProfilePackagingModes::DoNotPackage)
+		{
+			return EContentScheme::DevelopmentPackage;
+		}
+		else if (Profile->IsPackingWithUnrealPak())
+		{
+			return EContentScheme::PakFiles;
+		}
+		else if (Profile->IsUsingZenPakStreaming())
+		{
+			return EContentScheme::ZenPakStreaming;
+		}
+		else if (Profile->IsUsingZenStore())
+		{
+			return EContentScheme::ZenStreaming;
+		}
+		else
+		{
+			return EContentScheme::LooseFiles;
+		}
+	}
+
+
+
+	void FModel::SetProfileContentScheme(EContentScheme ContentScheme, const ILauncherProfileRef& Profile, ELauncherProfileDeploymentModes::Type DefaultDeploymentMode)
+	{
+		bool bPakFiles = (ContentScheme == EContentScheme::PakFiles || ContentScheme == EContentScheme::DevelopmentPackage);
+		bool bUseZen = (ContentScheme != EContentScheme::LooseFiles && ContentScheme != EContentScheme::CookOnTheFly);
+		bool bCOTF = (ContentScheme == EContentScheme::CookOnTheFly);
+		bool bPackage = (ContentScheme == EContentScheme::DevelopmentPackage);
+		bool bZenPakStreaming = (ContentScheme == EContentScheme::ZenPakStreaming);
+		// FIXME: may need to turn off Zen/Pak in UAT somehow if it's set in project defaults?
+
+		Profile->SetUseZenPakStreaming(bZenPakStreaming);
+		Profile->SetDeployWithUnrealPak(bPakFiles);
+		Profile->SetUseZenStore(bUseZen);
+
+		if (bCOTF)
+		{
+			Profile->SetCookMode(ELauncherProfileCookModes::OnTheFly);
+		}
+		else if (bZenPakStreaming)
+		{
+			Profile->SetCookMode(ELauncherProfileCookModes::DoNotCook);
+		}
+		else
+		{
+			Profile->SetCookMode(ELauncherProfileCookModes::ByTheBook);
+		}
+
+		if (bPackage)
+		{
+			Profile->SetDeploymentMode(ELauncherProfileDeploymentModes::DoNotDeploy); // @todo: some platforms support package deployment. should update verification to check new function in ITargetPlatformControls
+			Profile->SetPackagingMode(ELauncherProfilePackagingModes::Locally);
+		}
+		else
+		{
+			Profile->SetDeploymentMode(bCOTF ? ELauncherProfileDeploymentModes::FileServer : DefaultDeploymentMode);
+			Profile->SetPackagingMode(ELauncherProfilePackagingModes::DoNotPackage);
+		}
+	}
+
+
+
+	bool FModel::IsBasicLaunchProfile( const ILauncherProfilePtr& Profile ) const
+	{
+		return (Profile == BasicLaunchProfile);
+	}
+
+
+	TSharedPtr<FLaunchLogMessage> FModel::AddLogMessage( const FString& InMessage, ELogVerbosity::Type InVerbosity )
+	{
+		TSharedPtr<FLaunchLogMessage> Message = MakeShared<FLaunchLogMessage>(InMessage, InVerbosity);
+		LaunchLogMessages.Add(Message);
+		return Message;
+	}
+
+	void FModel::ClearLogMessages()
+	{
+		LaunchLogMessages.Reset();
+	}
+
+
+
+	TArray<EContentScheme> GetAllContentSchemes()
+	{
+		TArray<EContentScheme> ContentSchemes;
+		for (uint8 Index = 0; Index < (uint8)EContentScheme::MAX; Index++)
+		{
+			ContentSchemes.Add((EContentScheme)Index);
+		}
+
+		return MoveTemp(ContentSchemes);
+	}
+
+	FText GetContentSchemeDisplayName(EContentScheme ContentScheme)
+	{
+		switch (ContentScheme)
+		{
+			case EContentScheme::PakFiles:           return LOCTEXT("ContentSchemePakFiles","Pak Files"); break;
+			case EContentScheme::ZenStreaming:       return LOCTEXT("ContentSchemeZenStreaming","Zen Streaming"); break;
+			case EContentScheme::ZenPakStreaming:    return LOCTEXT("ContentSchemeZenPakStreaming","Zen Pak Streaming"); break;
+			case EContentScheme::DevelopmentPackage: return LOCTEXT("ContentSchemeDevPackage","Development Package"); break;
+			case EContentScheme::LooseFiles:         return LOCTEXT("ContentSchemeLooseFiles","Loose Files (legacy)"); break;
+			case EContentScheme::CookOnTheFly:       return LOCTEXT("ContentSchemeCOTF","Cook On The Fly"); break;
+		}
+
+		checkNoEntry();
+		return FText::GetEmpty();
+
+	}
+
+	FText GetContentSchemeToolTip(EContentScheme ContentScheme)
+	{
+		switch (ContentScheme)
+		{
+			case EContentScheme::PakFiles:           return LOCTEXT("ContentSchemeTipPakFiles","Store cooked game content in one ore more large Pak Files"); break;
+			case EContentScheme::ZenStreaming:       return LOCTEXT("ContentSchemeTipZenStreaming","Stream cooked game content from Zen Server"); break;
+			case EContentScheme::ZenPakStreaming:    return LOCTEXT("ContentSchemeTipZenPakStreaming","Stream an existing Pak Files build via a Zen"); break;
+			case EContentScheme::DevelopmentPackage: return LOCTEXT("ContentSchemeTipDevPackage","Package cooked game content into a single installable package file for development purposes, where available"); break;
+			case EContentScheme::LooseFiles:         return LOCTEXT("ContentSchemeTipLooseFiles","Store cooked game assets in individual files (legacy - recommend moving to Zen Streaming. This option will not work if the project is already configured with 'Use Zen Store')"); break;
+			case EContentScheme::CookOnTheFly:       return LOCTEXT("ContentSchemeTipCOTF","Only cook game assets when the game requires them, and send them over the network (legacy - slow)"); break;
+		}
+
+		checkNoEntry();
+		return FText::GetEmpty();
+
+	}
+
+
+
+
+	const TCHAR* LexToString( const CustomLaunchUI::EContentScheme& ContentScheme)
+	{
+		switch (ContentScheme)
+		{
+			case CustomLaunchUI::EContentScheme::PakFiles:           return TEXT("PakFiles");
+			case CustomLaunchUI::EContentScheme::ZenStreaming:       return TEXT("ZenStreaming");
+			case CustomLaunchUI::EContentScheme::ZenPakStreaming:    return TEXT("ZenPakStreaming");
+			case CustomLaunchUI::EContentScheme::DevelopmentPackage: return TEXT("DevelopmentPackage");
+			case CustomLaunchUI::EContentScheme::LooseFiles:         return TEXT("LooseFiles");
+			case CustomLaunchUI::EContentScheme::CookOnTheFly:       return TEXT("CookOnTheFly");
+			default: return TEXT("Unknown");
+		}
+	}
+
+
+
+	bool LexTryParseString( CustomLaunchUI::EContentScheme& OutContentScheme, const TCHAR* String )
+	{
+		if (FCString::Stricmp(String, TEXT("PakFiles")) == 0) 
+		{
+			OutContentScheme = CustomLaunchUI::EContentScheme::PakFiles;
+			return true;
+		}
+		else if (FCString::Stricmp(String, TEXT("ZenStreaming")) == 0) 
+		{
+			OutContentScheme = CustomLaunchUI::EContentScheme::ZenStreaming;
+			return true;
+		}
+		else if (FCString::Stricmp(String, TEXT("ZenPakStreaming")) == 0) 
+		{
+			OutContentScheme = CustomLaunchUI::EContentScheme::ZenPakStreaming;
+			return true;
+		}
+		else if (FCString::Stricmp(String, TEXT("DevelopmentPackage")) == 0) 
+		{
+			OutContentScheme = CustomLaunchUI::EContentScheme::DevelopmentPackage;
+			return true;
+		}
+		else if (FCString::Stricmp(String, TEXT("LooseFiles")) == 0) 
+		{
+			OutContentScheme = CustomLaunchUI::EContentScheme::LooseFiles;
+			return true;
+		}
+		else if (FCString::Stricmp(String, TEXT("CookOnTheFly")) == 0) 
+		{
+			OutContentScheme = CustomLaunchUI::EContentScheme::CookOnTheFly;
+			return true;
+		}
+		else
+		{
+			return false;
+		}
+	}
+
+
+
+	FText GetProfileLaunchErrorMessage(ILauncherProfilePtr Profile)
+	{
+		if (!Profile.IsValid())
+		{
+			return LOCTEXT("LaunchErrNoProfileTip", "There is no profile selected");
+		}
+
+		if (Profile->IsValidForLaunch())
+		{
+			return LOCTEXT("LaunchProfileTip", "Launch this profile now");
+		}
+
+		FTextBuilder MsgTextBuilder;
+		MsgTextBuilder.AppendLine(LOCTEXT("LaunchErrValidation", "There are validation errors with this profile. Please fix them before launching:"));
+		MsgTextBuilder.Indent();
+		for (int i = 0; i < (int)ELauncherProfileValidationErrors::Count; i++)
+		{
+			ELauncherProfileValidationErrors::Type Error = (ELauncherProfileValidationErrors::Type)i;
+			if (Profile->HasValidationError(Error))
+			{
+				MsgTextBuilder.AppendLine(LexToStringLocalized(Error));
+			}
+		}
+		MsgTextBuilder.Unindent();
+
+		return MsgTextBuilder.ToText();
+	}
+}
+
+
+
+#undef LOCTEXT_NAMESPACE
diff --git a/Engine/Plugins/Developer/CustomLaunchUI/Source/Private/Model/CustomLaunchUIModel.h b/Engine/Plugins/Developer/CustomLaunchUI/Source/Private/Model/CustomLaunchUIModel.h
new file mode 100644
--- /dev/null
+++ b/Engine/Plugins/Developer/CustomLaunchUI/Source/Private/Model/CustomLaunchUIModel.h
@@ -0,0 +1,164 @@
+// Copyright Epic Games, Inc. All Rights Reserved.
+
+#pragma once
+
+#include "CoreTypes.h"
+#include "Containers/Array.h"
+#include "Delegates/Delegate.h"
+#include "ILauncherProfileManager.h"
+#include "Templates/SharedPointer.h"
+#include "Misc/DataDrivenPlatformInfoRegistry.h"
+
+
+
+namespace PlatformInfo
+{
+	struct FTargetPlatformInfo;
+};
+
+class ITargetDeviceProxyManager;
+class ITargetDeviceProxy;
+class ILauncher;
+struct FTargetInfo;
+
+namespace CustomLaunchUI
+{
+	enum class EContentScheme : uint8
+	{
+		PakFiles,
+		ZenStreaming,
+		ZenPakStreaming,
+		DevelopmentPackage,
+		LooseFiles,
+		CookOnTheFly,
+		MAX
+	};
+	TArray<EContentScheme> GetAllContentSchemes();
+	FText GetContentSchemeDisplayName(EContentScheme ContentScheme);
+	FText GetContentSchemeToolTip(EContentScheme ContentScheme);
+	FText GetProfileLaunchErrorMessage(ILauncherProfilePtr Profile);
+	const TCHAR* LexToString( const CustomLaunchUI::EContentScheme& ContentScheme);
+	bool LexTryParseString( CustomLaunchUI::EContentScheme& OutContentScheme, const TCHAR* String );
+
+	struct FLaunchLogMessage
+	{
+		TSharedRef<FString> Message;
+		ELogVerbosity::Type Verbosity;
+
+		FLaunchLogMessage(const FString& InMessage, ELogVerbosity::Type InVerbosity)
+			: Message(MakeShared<FString>(InMessage))
+			, Verbosity(InVerbosity)
+		{}
+	};
+
+
+
+	DECLARE_DELEGATE_OneParam(FOnProfileClicked, const ILauncherProfilePtr&)
+
+	DECLARE_MULTICAST_DELEGATE_TwoParams(FOnSelectedProfileChanged, const ILauncherProfilePtr& /*new*/, const ILauncherProfilePtr& /*old*/)
+	DECLARE_MULTICAST_DELEGATE(FOnSelectedProjectChanged)
+	DECLARE_MULTICAST_DELEGATE(FOnBuildTargetChanged)
+	DECLARE_MULTICAST_DELEGATE(FOnSelectedPlatformChanged)
+
+	extern bool bUseFriendlyBuildTargetSelection;
+
+
+	class FModel : public TSharedFromThis<FModel>
+	{
+	public:
+		FModel(const TSharedRef<ITargetDeviceProxyManager>& InDeviceProxyManager, const TSharedRef<ILauncher>& InLauncher, const TSharedRef<ILauncherProfileManager>& InProfileManager);
+		~FModel();
+
+
+		const TSharedRef<ITargetDeviceProxyManager>& GetDeviceProxyManager() const
+		{
+			return DeviceProxyManager;
+		}
+
+		const TSharedRef<ILauncher>& GetLauncher() const
+		{
+			return Launcher;
+		}
+
+		const TSharedRef<ILauncherProfileManager>& GetProfileManager() const
+		{
+			return ProfileManager;
+		}
+
+
+		const TArray<ILauncherProfilePtr>& GetAllProfiles() const
+		{
+			return AllProfiles;
+		}
+
+		ILauncherProfilePtr GetBasicLaunchProfile() const
+		{
+			return BasicLaunchProfile;
+		}
+
+		const ILauncherProfilePtr& GetSelectedProfile() const
+		{
+			return SelectedProfile;
+		}
+
+		FOnSelectedProfileChanged& OnProfileSelected()
+		{
+			return ProfileSelectedDelegate;
+		}
+
+		void SelectProfile(const ILauncherProfilePtr& Profile);
+		void SortProfiles();
+
+		bool IsAdvancedProfile(const ILauncherProfileRef& Profile) const;
+		bool IsBasicLaunchProfile( const ILauncherProfilePtr& Profile ) const;
+
+		ILauncherProfileRef CreateCustomProfile( const TCHAR* Name );
+		ILauncherProfileRef CreateBasicLaunchProfile();
+		ILauncherProfilePtr CloneCustomProfile(const ILauncherProfileRef& Profile);
+
+		EContentScheme DetermineProfileContentScheme(const ILauncherProfileRef& Profile) const;
+		void SetProfileContentScheme(EContentScheme ContentScheme, const ILauncherProfileRef& Profile, ELauncherProfileDeploymentModes::Type DefaultDeploymentMode = ELauncherProfileDeploymentModes::CopyToDevice);
+
+		TSharedPtr<FLaunchLogMessage> AddLogMessage( const FString& InMessage, ELogVerbosity::Type InVerbosity );
+		void ClearLogMessages();
+		int32 GetNumLogMessages() const { return LaunchLogMessages.Num(); }
+
+
+		static const PlatformInfo::FTargetPlatformInfo* GetPlatformInfo(const ILauncherProfilePtr& Profile);
+		static const PlatformInfo::FTargetPlatformInfo* GetPlatformInfo(FName PlatformName, const FTargetInfo& BuildTargetInfo);
+
+		static FTargetInfo GetBuildTargetInfo( const FString& BuildTargetName, const FString& ProjectPath );
+		static FTargetInfo GetBuildTargetInfo( const ILauncherProfileRef& Profile );
+
+		static TSharedPtr<ITargetDeviceProxy> GetDeviceProxy(const ILauncherProfileRef& Profile);
+		static void UpdatedCookedPlatformsFromDeployDeviceProxy(const ILauncherProfileRef& Profile, TSharedPtr<ITargetDeviceProxy> DeviceProxy = nullptr);
+		static void UpdateCookedPlatformsFromBuildTarget(const ILauncherProfileRef& Profile);
+
+		static FString GetVanillaPlatformName( const FString& PlatformName );
+		static FString GetBuildTargetPlatformName( const FString& PlatformName, const FTargetInfo& BuildTargetInfo );
+
+	private:
+		void LoadConfig();
+		void SaveConfig();
+
+		void HandleProfileManagerProfileAdded(const ILauncherProfileRef& Profile);
+		void HandleProfileManagerProfileRemoved(const ILauncherProfileRef& Profile);
+		void HandleDeviceProxyAdded(const TSharedRef<ITargetDeviceProxy>& DeviceProxy );
+		void HandleDeviceProxyRemoved(const TSharedRef<ITargetDeviceProxy>& DeviceProxy );
+
+		TSharedRef<ITargetDeviceProxyManager> DeviceProxyManager;
+		TSharedRef<ILauncher> Launcher;
+		TSharedRef<ILauncherProfileManager> ProfileManager;
+
+		ILauncherProfilePtr SelectedProfile;
+		TArray<ILauncherProfilePtr> AllProfiles;
+		ILauncherProfilePtr BasicLaunchProfile;
+		bool bHasSetBasicLaunchProfilePlatform = false;
+
+		friend class FLaunchLogTextLayoutMarshaller;
+		TArray<TSharedPtr<FLaunchLogMessage>> LaunchLogMessages;
+
+		FOnSelectedProfileChanged ProfileSelectedDelegate;
+		FOnSelectedPlatformChanged PlatformChangedDelegate;
+	};
+}
diff --git a/Engine/Plugins/Developer/CustomLaunchUI/Source/Private/Styling/CustomLaunchUIStyle.cpp b/Engine/Plugins/Developer/CustomLaunchUI/Source/Private/Styling/CustomLaunchUIStyle.cpp
new file mode 100644
--- /dev/null
+++ b/Engine/Plugins/Developer/CustomLaunchUI/Source/Private/Styling/CustomLaunchUIStyle.cpp
@@ -0,0 +1,143 @@
+// Copyright Epic Games, Inc. All Rights Reserved.
+
+#include "Styling/CustomLaunchUIStyle.h"
+
+#include "Styling/StyleColors.h"
+#include "Styling/SlateStyleRegistry.h"
+#include "Styling/SlateTypes.h"
+#include "Styling/CoreStyle.h"
+#include "Styling/SlateStyleMacros.h"
+#include "Brushes/SlateBoxBrush.h"
+#include "Brushes/SlateImageBrush.h"
+#include "Brushes/SlateNoResource.h"
+#include "Misc/Paths.h"
+#include "Interfaces/IPluginManager.h"
+#include "PlatformInfo.h"
+
+// This is to fix the issue that SlateStyleMacros like IMAGE_BRUSH look for RootToContentDir but StyleSet->RootToContentDir is how this style is set up
+#define RootToContentDir StyleSet->RootToContentDir
+#define RootToCoreContentDir StyleSet->RootToCoreContentDir
+
+TSharedPtr<FSlateStyleSet> FCustomLaunchUIStyle::StyleSet = nullptr;
+
+static const FVector2D Icon15x15(16.0f, 16.0f);
+static const FVector2D Icon16x16(16.0f, 16.0f);
+static const FVector2D Icon28x28(28.0f, 28.0f);
+static const FVector2D Icon36x36(36.0f, 36.0f);
+static const FVector2D Icon512x512(512.0f, 512.0f);
+
+
+void FCustomLaunchUIStyle::Initialize()
+{
+	if (StyleSet.IsValid())
+	{
+		return;
+	}
+
+	StyleSet = MakeShareable(new FSlateStyleSet("CustomLaunchUIStyle") );
+	StyleSet->SetContentRoot(IPluginManager::Get().FindPlugin("CustomLaunchUI")->GetBaseDir() / TEXT("Resources"));
+	StyleSet->SetCoreContentRoot(FPaths::EngineContentDir() / TEXT("Slate"));
+
+	StyleSet->Set("WhiteGroupBorder", new CORE_BOX_BRUSH(TEXT("Common/WhiteGroupBorder"), FMargin(4.0f/16.0f)));
+
+	StyleSet->Set("SidePanelRightClose", new IMAGE_BRUSH_SVG(TEXT("SidePanelRightClose"), Icon16x16));
+	StyleSet->Set("Icons.ClearLog", new IMAGE_BRUSH_SVG(TEXT("ClearLog"), Icon16x16));
+
+	StyleSet->Set("BadgeOutlined.Error", new IMAGE_BRUSH_SVG(TEXT("BadgeOutlinedError_15"), Icon15x15));
+	StyleSet->Set("BadgeOutlined.AllComplete", new IMAGE_BRUSH_SVG(TEXT("BadgeOutlinedSuccess_15-1"), Icon15x15));
+	StyleSet->Set("BadgeOutlined.Success", new IMAGE_BRUSH_SVG(TEXT("BadgeOutlinedSuccess_15"), Icon15x15));
+	StyleSet->Set("OuterCircle", new IMAGE_BRUSH_SVG(TEXT("outer_circle_white"), Icon36x36));
+	/*StyleSet->Set("OuterCircle.Busy", new IMAGE_BRUSH(TEXT("outer_circle_busy"), Icon36x36));*/ // unused - higher res image has less aliasing when rotating
+	StyleSet->Set("OuterCircle.Busy", new IMAGE_BRUSH(TEXT("outer_circle_busy_512px"), Icon512x512));
+	StyleSet->Set("FullCircle", new IMAGE_BRUSH_SVG(TEXT("full_circle_white"), Icon36x36));
+
+	StyleSet->Set("State.Busy", FLinearColor(FColor(0xFF0070E0)) );
+	StyleSet->Set("State.Pending", FLinearColor(FColor(0xFF1A1A1A)) );
+	StyleSet->Set("State.Error", FLinearColor(FColor(0xFFEF3535)) );
+	StyleSet->Set("State.Complete", FLinearColor(FColor(0xFF0070E0)) );
+	StyleSet->Set("State.Canceled", FLinearColor(FColor(0xFF383838)) );
+	StyleSet->Set("State.AllComplete", FLinearColor(FColor(0xFF0070E0)) );
+
+	
+	FSlateStyleRegistry::RegisterSlateStyle( *StyleSet.Get() );
+}
+
+#undef RootToContentDir
+#undef RootToCoreContentDir
+
+void FCustomLaunchUIStyle::Shutdown()
+{
+	if (StyleSet.IsValid())
+	{
+		FSlateStyleRegistry::UnRegisterSlateStyle( *StyleSet.Get() );
+		ensure( StyleSet.IsUnique() );
+		StyleSet.Reset();
+	}
+}
+
+const ISlateStyle& FCustomLaunchUIStyle::Get()
+{
+	return *(StyleSet.Get());
+}
+
+const FName& FCustomLaunchUIStyle::GetStyleSetName()
+{
+	return StyleSet->GetStyleSetName();
+}
+
+
+
+const FSlateBrush* FCustomLaunchUIStyle::GetBrushForTask(ILauncherTaskPtr Task)
+{
+	// fixme: use bespoke icons
+	if (Task.IsValid())
+	{
+		FString TaskName = Task->GetName();
+		if (TaskName.Contains(TEXT("Cooking in the editor")))
+		{
+			return FCoreStyle::Get().GetBrush("GraphEditor.Timeline_16x");
+		}
+		else if (TaskName.Contains(TEXT("Post Launch")))
+		{
+			return FCoreStyle::Get().GetBrush("GraphEditor.Clean");
+		}
+		else if (TaskName.Contains(TEXT("Launch")))
+		{
+			return FCoreStyle::Get().GetBrush("Icons.Launch");
+		}
+		else if (TaskName.Contains(TEXT("Build")))
+		{
+			return FCoreStyle::Get().GetBrush("Icons.Adjust");
+		}
+		else if (TaskName.Contains(TEXT("Cook")))
+		{
+			return FCoreStyle::Get().GetBrush("MainFrame.CookContent");
+		}
+		else if (TaskName.Contains(TEXT("Deploy")))
+		{
+			return FCoreStyle::Get().GetBrush("DeviceDetails.TabIcon");
+		}
+		else if (TaskName.Contains(TEXT("Archive")))
+		{
+			return FCoreStyle::Get().GetBrush("Icons.Save");
+		}
+		else if (TaskName.Contains(TEXT("Run")))
+		{
+			return FCoreStyle::Get().GetBrush("GraphEditor.Timeline_16x");
+		}
+
+	}
+
+	return FCoreStyle::Get().GetBrush("Icons.Package");
+}
+
+
+const FSlateBrush* FCustomLaunchUIStyle::GetProfileBrushForPlatform(const PlatformInfo::FTargetPlatformInfo* PlatformInfo, EPlatformIconSize IconSize)
+{
+	if (PlatformInfo != nullptr)
+	{
+		return FAppStyle::GetBrush(PlatformInfo->GetIconStyleName(IconSize));
+	}
+	
+	return FAppStyle::GetBrush("LauncherCommand.QuickLaunch");
+}
diff --git a/Engine/Plugins/Developer/CustomLaunchUI/Source/Private/Styling/CustomLaunchUIStyle.h b/Engine/Plugins/Developer/CustomLaunchUI/Source/Private/Styling/CustomLaunchUIStyle.h
new file mode 100644
--- /dev/null
+++ b/Engine/Plugins/Developer/CustomLaunchUI/Source/Private/Styling/CustomLaunchUIStyle.h
@@ -0,0 +1,32 @@
+// Copyright Epic Games, Inc. All Rights Reserved.
+
+#pragma once
+
+#include "Styling/SlateStyle.h"
+#include "ILauncherTask.h"
+#include "Misc/DataDrivenPlatformInfoRegistry.h"
+
+namespace PlatformInfo
+{
+	struct FTargetPlatformInfo;
+};
+
+class FCustomLaunchUIStyle
+{
+public:
+	static void Initialize();
+	static void Shutdown();
+
+	static const ISlateStyle& Get();
+	static const FName& GetStyleSetName();
+	static const FSlateBrush* GetBrush(FName PropertyName, const ANSICHAR* Specifier = nullptr)
+	{
+		return StyleSet->GetBrush(PropertyName, Specifier);
+	}
+
+	static const FSlateBrush* GetBrushForTask(ILauncherTaskPtr Task);
+	static const FSlateBrush* GetProfileBrushForPlatform(const PlatformInfo::FTargetPlatformInfo* PlatformInfo, EPlatformIconSize IconSize);
+
+private:
+	static TSharedPtr<FSlateStyleSet> StyleSet;	
+};
diff --git a/Engine/Plugins/Developer/CustomLaunchUI/Source/Private/Widgets/Output/CustomLaunchOutputLogMarshaller.cpp b/Engine/Plugins/Developer/CustomLaunchUI/Source/Private/Widgets/Output/CustomLaunchOutputLogMarshaller.cpp
new file mode 100644
--- /dev/null
+++ b/Engine/Plugins/Developer/CustomLaunchUI/Source/Private/Widgets/Output/CustomLaunchOutputLogMarshaller.cpp
@@ -0,0 +1,136 @@
+// Copyright Epic Games, Inc. All Rights Reserved.
+
+#include "Widgets/Output/CustomLaunchOutputLogMarshaller.h"
+#include "Model/CustomLaunchUIModel.h"
+#include "Styling/AppStyle.h"
+#include "Framework/Text/SlateTextRun.h"
+#include "Framework/Text/ISlateLineHighlighter.h"
+
+
+namespace CustomLaunchUI
+{
+	FLaunchLogTextLayoutMarshaller::FLaunchLogTextLayoutMarshaller(const TSharedRef<CustomLaunchUI::FModel>& InModel) 
+		: Model(InModel)
+		, TextLayout(nullptr)
+	{
+		MessageStyle = FAppStyle::Get().GetWidgetStyle<FTextBlockStyle>("MonospacedText");
+
+		DisplayStyle = MessageStyle;
+		DisplayStyle.ColorAndOpacity = FSlateColor(FColor::Green);
+
+		WarningStyle = MessageStyle;
+		WarningStyle.ColorAndOpacity = FSlateColor(FColor::Yellow);
+
+		ErrorStyle = MessageStyle;
+		ErrorStyle.ColorAndOpacity = FSlateColor(FColor::Red);
+	}
+
+
+	void FLaunchLogTextLayoutMarshaller::SetText(const FString& SourceString, FTextLayout& TargetTextLayout)
+	{
+		TextLayout = &TargetTextLayout;
+		FlushPendingLogMessages();
+	}
+
+
+	void FLaunchLogTextLayoutMarshaller::GetText(FString& TargetString, const FTextLayout& SourceTextLayout)
+	{
+		SourceTextLayout.GetAsText(TargetString);
+	}
+
+
+	void FLaunchLogTextLayoutMarshaller::MakeDirty()
+	{
+		FBaseTextLayoutMarshaller::MakeDirty();
+		NumFilteredMessages = 0;
+	}
+
+
+	void FLaunchLogTextLayoutMarshaller::AddPendingLogMessage( TSharedPtr<FLaunchLogMessage> Message )
+	{
+		PendingMessages.Enqueue(Message);
+	}
+
+
+	void FLaunchLogTextLayoutMarshaller::RefreshAllLogMessages()
+	{
+		MakeDirty();
+
+		for (TSharedPtr<FLaunchLogMessage> Message : Model->LaunchLogMessages)
+		{
+			PendingMessages.Enqueue(Message);
+		}
+	}
+
+
+	void FLaunchLogTextLayoutMarshaller::SetFilter( ELogFilter InLogFilter )
+	{
+		LogFilter = InLogFilter;
+		RefreshAllLogMessages();
+	}
+
+
+	void FLaunchLogTextLayoutMarshaller::SetFilterString( const FString& InLogFilterString )
+	{
+		LogFilterString = InLogFilterString;
+		RefreshAllLogMessages();
+	}
+
+
+	void FLaunchLogTextLayoutMarshaller::FlushPendingLogMessages()
+	{
+		bool bRefreshLog = false;
+
+		TArray<FTextLayout::FNewLineData> LinesToAdd;
+
+		TSharedPtr<FLaunchLogMessage> Message;
+		while (PendingMessages.Dequeue(Message))
+		{
+			if (LogFilter == ELogFilter::Errors && Message->Verbosity != ELogVerbosity::Error && Message->Verbosity != ELogVerbosity::Fatal)
+			{
+				continue;
+			}
+			else if (LogFilter == ELogFilter::WarningsAndErrors && Message->Verbosity != ELogVerbosity::Warning && Message->Verbosity != ELogVerbosity::Error && Message->Verbosity != ELogVerbosity::Fatal)
+			{
+				continue;
+			}
+			else if (!LogFilterString.IsEmpty() && !Message->Message->Contains(LogFilterString))
+			{
+				continue;
+			}
+
+			NumFilteredMessages++;
+			LinesToAdd.Emplace(MoveTemp(Message->Message), GetRun(Message));
+
+			bRefreshLog = true;
+		}
+
+		if (bRefreshLog)
+		{
+			TextLayout->AddLines(LinesToAdd);
+		}
+	}
+
+
+	TArray<TSharedRef<IRun>> FLaunchLogTextLayoutMarshaller::GetRun(TSharedPtr<FLaunchLogMessage> Message) const
+	{
+		const FTextBlockStyle* TextStyle = &MessageStyle;
+		if (Message->Verbosity == ELogVerbosity::Warning)
+		{
+			TextStyle = &WarningStyle;
+		}
+		else if (Message->Verbosity == ELogVerbosity::Error || Message->Verbosity == ELogVerbosity::Fatal)
+		{
+			TextStyle = &ErrorStyle;
+		}
+		else if (Message->Verbosity == ELogVerbosity::Display)
+		{
+			TextStyle = &DisplayStyle;
+		}
+
+		TArray<TSharedRef<IRun>> Runs;
+		Runs.Add(FSlateTextRun::Create(FRunInfo(), Message->Message, *TextStyle));
+
+		return MoveTemp(Runs);
+	}
+}
diff --git a/Engine/Plugins/Developer/CustomLaunchUI/Source/Private/Widgets/Output/CustomLaunchOutputLogMarshaller.h b/Engine/Plugins/Developer/CustomLaunchUI/Source/Private/Widgets/Output/CustomLaunchOutputLogMarshaller.h
new file mode 100644
--- /dev/null
+++ b/Engine/Plugins/Developer/CustomLaunchUI/Source/Private/Widgets/Output/CustomLaunchOutputLogMarshaller.h
@@ -0,0 +1,63 @@
+// Copyright Epic Games, Inc. All Rights Reserved.
+
+#pragma once
+
+#include "CoreTypes.h"
+#include "Containers/Array.h"
+#include "Containers/SpscQueue.h"
+#include "Templates/SharedPointer.h"
+#include "Styling/SlateTypes.h"
+#include "Framework/Text/TextLayout.h"
+#include "Framework/Text/BaseTextLayoutMarshaller.h"
+
+
+namespace CustomLaunchUI
+{
+	class FModel;
+	struct FLaunchLogMessage;
+
+	enum class ELogFilter : uint8
+	{
+		All,
+		WarningsAndErrors,
+		Errors,
+	};
+
+
+	// this class is responsible for marshalling the output log from FModel into an associated multiline text box
+	class FLaunchLogTextLayoutMarshaller : public FBaseTextLayoutMarshaller
+	{
+	public:
+		FLaunchLogTextLayoutMarshaller(const TSharedRef<CustomLaunchUI::FModel>& InModel);
+		virtual void SetText(const FString& SourceString, FTextLayout& TargetTextLayout) override;
+		virtual void GetText(FString& TargetString, const FTextLayout& SourceTextLayout) override;
+		virtual void MakeDirty() override;
+
+		ELogFilter GetFilter() const { return LogFilter; }
+		void SetFilter( ELogFilter InLogFilter );
+
+		FString GetFilterString() const { return LogFilterString; }
+		void SetFilterString( const FString& InLogFilterString );
+
+		void AddPendingLogMessage( TSharedPtr<FLaunchLogMessage> Message );
+		void FlushPendingLogMessages();
+
+		void RefreshAllLogMessages();
+		int32 GetNumFilteredMessages() const { return NumFilteredMessages; }
+	protected:
+		TArray<TSharedRef<class IRun>> GetRun(TSharedPtr<FLaunchLogMessage> Message) const;
+
+		FTextBlockStyle MessageStyle;
+		FTextBlockStyle DisplayStyle;
+		FTextBlockStyle WarningStyle;
+		FTextBlockStyle ErrorStyle;
+
+		const TSharedRef<CustomLaunchUI::FModel> Model;
+		TSpscQueue<TSharedPtr<FLaunchLogMessage>> PendingMessages;
+		ELogFilter LogFilter = ELogFilter::All;
+		FString LogFilterString;
+		FTextLayout* TextLayout;
+		int32 NumFilteredMessages = 0;
+
+	};
+}
diff --git a/Engine/Plugins/Developer/CustomLaunchUI/Source/Private/Widgets/Output/SCustomLaunchLaunchPanel.cpp b/Engine/Plugins/Developer/CustomLaunchUI/Source/Private/Widgets/Output/SCustomLaunchLaunchPanel.cpp
new file mode 100644
--- /dev/null
+++ b/Engine/Plugins/Developer/CustomLaunchUI/Source/Private/Widgets/Output/SCustomLaunchLaunchPanel.cpp
@@ -0,0 +1,675 @@
+// Copyright Epic Games, Inc. All Rights Reserved.
+
+#include "Widgets/Output/SCustomLaunchLaunchPanel.h"
+
+#include "Styling/AppStyle.h"
+#include "Styling/CustomLaunchUIStyle.h"
+#include "SlateOptMacros.h"
+#include "Widgets/SBoxPanel.h"
+#include "Widgets/Text/STextBlock.h"
+#include "Widgets/Layout/SSpacer.h"
+#include "Widgets/Layout/SGridPanel.h"
+#include "Widgets/Images/SImage.h"
+#include "Widgets/Input/SButton.h"
+#include "Widgets/Notifications/SProgressBar.h"
+#include "Widgets/Output/SCustomLaunchOutputLog.h"
+#include "ITargetDeviceProxy.h"
+#include "PlatformInfo.h"
+
+
+
+#define LOCTEXT_NAMESPACE "SCustomLaunchLaunchPanel"
+
+
+
+BEGIN_SLATE_FUNCTION_BUILD_OPTIMIZATION
+void SCustomLaunchLaunchPanel::Construct(const FArguments& InArgs, const TSharedRef<CustomLaunchUI::FModel>& InModel)
+{
+	Model = InModel;
+	OnCloseClicked = InArgs._OnCloseClicked;
+	OnRerunClicked = InArgs._OnRerunClicked;
+
+	LaunchLogTextMarshaller = MakeShared<CustomLaunchUI::FLaunchLogTextLayoutMarshaller>(InModel);
+	OutputLog = SNew(SCustomLaunchOutputLog, InModel, LaunchLogTextMarshaller.ToSharedRef());
+
+	ChildSlot
+	[
+		SNew(SVerticalBox)
+
+		// top banner
+		+SVerticalBox::Slot()
+		.AutoHeight()
+		.Padding(8)
+		[
+			SNew(SBorder)
+			.BorderImage(FAppStyle::GetBrush("Brushes.Header"))
+			.Padding(16)
+			[
+				SNew(SVerticalBox)
+				
+				// profile details & control buttons
+				+SVerticalBox::Slot()
+				.Padding(0)
+				[
+					SNew(SHorizontalBox)
+
+					// profile icon
+					+SHorizontalBox::Slot()
+					.AutoWidth()
+					.Padding(4,0)
+					.VAlign(VAlign_Center)
+					[
+						SNew(SImage)
+						.DesiredSizeOverride(FVector2D(44,44))
+						.Image( this, &SCustomLaunchLaunchPanel::GetProfileImage )
+					]
+
+					// profile details
+					+SHorizontalBox::Slot()
+					.AutoWidth()
+					.Padding(4)
+					.VAlign(VAlign_Center)
+					[
+						SNew(SVerticalBox)
+
+						// profile name
+						+SVerticalBox::Slot()
+						.AutoHeight()
+						.Padding(2)
+						[
+							SNew(STextBlock)
+							.Text( this, &SCustomLaunchLaunchPanel::GetProfileName )
+						]
+
+						// profile description
+						+SVerticalBox::Slot()
+						.AutoHeight()
+						.Padding(2)
+						[
+							SNew(STextBlock)
+							.Text( this, &SCustomLaunchLaunchPanel::GetProfileDescription )
+							.ColorAndOpacity(FSlateColor::UseSubduedForeground())
+						]
+					]
+
+					// spacer
+					+SHorizontalBox::Slot()
+					.FillWidth(1)
+					[
+						SNew(SSpacer)
+					]
+
+					// control buttons
+					+SHorizontalBox::Slot()
+					.Padding(4,0,0,0)
+					.AutoWidth()
+					.VAlign(VAlign_Top)
+					[
+						SNew(SHorizontalBox)
+
+						// cancel button
+						+SHorizontalBox::Slot()
+						.Padding(FMargin(4,0,0,0))
+						[
+							SNew(SButton)
+							.OnClicked( this, &SCustomLaunchLaunchPanel::OnCancelButtonClicked )
+							.Visibility_Lambda( [this]() { return !IsIdle() ? EVisibility::Visible : EVisibility::Collapsed; } )
+							.IsEnabled_Lambda( [this]() { return IsBuilding(); } )
+							.Text(LOCTEXT("CancelButtonLabel", "Cancel"))
+							.ContentPadding(4)
+						]
+
+						// retry button
+						+SHorizontalBox::Slot()
+						.Padding(FMargin(4,0,0,0))
+						[
+							SNew(SButton)
+							.OnClicked( this, &SCustomLaunchLaunchPanel::OnRetryButtonClicked )
+							.Visibility_Lambda( [this]() { return IsIdle() ? EVisibility::Visible : EVisibility::Collapsed; } )
+							.Text(LOCTEXT("RetryButtonLabel", "Retry"))
+							.ContentPadding(4)
+						]
+
+						// done button
+						+SHorizontalBox::Slot()
+						.Padding(FMargin(4,0,0,0))
+						[
+							SNew(SButton)
+							.OnClicked( this, &SCustomLaunchLaunchPanel::OnDoneButtonClicked )
+							.Visibility_Lambda( [this]() { return IsIdle() ? EVisibility::Visible : EVisibility::Collapsed; } )
+							.ButtonStyle(FAppStyle::Get(), "PrimaryButton")
+							.Text(LOCTEXT("DoneButtonLabel", "Done"))
+							.ContentPadding(4)
+						]
+					]
+				]
+
+#if 0
+				// overall progress bar (only per-task granulatity)
+				+SVerticalBox::Slot()
+				.Padding(0,16)
+				.AutoHeight()
+				[
+					SNew(SProgressBar)
+					.Percent(this, &SCustomLaunchLaunchPanel::GetProgressPercent)
+				]
+#else
+				// experimental subway progress bar
+				+SVerticalBox::Slot()
+				.Padding(0,16)
+				.AutoHeight()
+				[
+					SAssignNew(SubwayProgressBar, SSegmentedProgressBar)
+				]
+#endif
+
+				// progress description
+				+SVerticalBox::Slot()
+				.Padding(0)
+				.AutoHeight()
+				[
+					SNew(SGridPanel)
+					.FillColumn(2, 1.0f)
+					
+					// project name
+					+SGridPanel::Slot(0,0)
+					.Padding(0,2)
+					[
+						SNew(SHorizontalBox)
+						+SHorizontalBox::Slot()
+						.AutoWidth()
+						.Padding(0,0,6,0)
+						[
+							SNew(STextBlock)
+							.Text(LOCTEXT("ProjectLabel", "Project"))
+							.ColorAndOpacity(FSlateColor::UseSubduedForeground())
+						]
+
+						+SHorizontalBox::Slot()
+						.AutoWidth()
+						[
+							SNew(STextBlock)
+							.Text( this, &SCustomLaunchLaunchPanel::GetProfileProjectName )
+						]
+					]
+
+					// configuration
+					+SGridPanel::Slot(1,0)
+					.Padding(32,2)
+					[
+						SNew(SHorizontalBox)
+						+SHorizontalBox::Slot()
+						.AutoWidth()
+						.Padding(0,0,6,0)
+						[
+							SNew(STextBlock)
+							.Text(LOCTEXT("ConfigurationLabel", "Configuration"))
+							.ColorAndOpacity(FSlateColor::UseSubduedForeground())
+						]
+
+						+SHorizontalBox::Slot()
+						.AutoWidth()
+						[
+							SNew(STextBlock)
+							.Text( this, &SCustomLaunchLaunchPanel::GetProfileConfigurationName )
+						]
+					]
+
+					// task progress
+					+SGridPanel::Slot(2,0)
+					.HAlign(HAlign_Right)
+					[
+						SNew(SHorizontalBox)
+						+SHorizontalBox::Slot()
+						.HAlign(HAlign_Right)
+						.AutoWidth()
+						.Padding(0,0,6,0)
+						[
+							SNew(STextBlock)
+							.Text( this, &SCustomLaunchLaunchPanel::GetCurrentTaskDescription )
+							.ColorAndOpacity(FSlateColor::UseSubduedForeground())
+						]
+
+						+SHorizontalBox::Slot()
+						.HAlign(HAlign_Right)
+						.AutoWidth()
+						[
+							SNew(STextBlock)
+							.Text( this, &SCustomLaunchLaunchPanel::GetProgressDescription )
+						]
+					]
+
+					// target name
+					+SGridPanel::Slot(0,1)
+					.Padding(0,2)
+					[
+						SNew(SHorizontalBox)
+						+SHorizontalBox::Slot()
+						.AutoWidth()
+						.Padding(0,0,6,0)
+						[
+							SNew(STextBlock)
+							.Text(LOCTEXT("TargetLabel", "Target"))
+							.ColorAndOpacity(FSlateColor::UseSubduedForeground())
+						]
+
+						+SHorizontalBox::Slot()
+						.AutoWidth()
+						[
+							SNew(STextBlock)
+							.Text( this, &SCustomLaunchLaunchPanel::GetProfileTargetName )
+						]
+					]
+					
+					// content scheme
+					+SGridPanel::Slot(1,1)
+					.Padding(32,2)
+					[
+						SNew(SHorizontalBox)
+						+SHorizontalBox::Slot()
+						.AutoWidth()
+						.Padding(0,0,6,0)
+						[
+							SNew(STextBlock)
+							.Text(LOCTEXT("ContentLabel", "Content"))
+							.ColorAndOpacity(FSlateColor::UseSubduedForeground())
+						]
+
+						+SHorizontalBox::Slot()
+						.AutoWidth()
+						[
+							SNew(STextBlock)
+							.Text( this, &SCustomLaunchLaunchPanel::GetProfileContentSchemeName )
+						]
+					]
+
+					// overall build duration
+					+SGridPanel::Slot(2,1)
+					.HAlign(HAlign_Right)
+					[
+						SNew(STextBlock)
+						.Text( this, &SCustomLaunchLaunchPanel::GetTotalDurationDescription )
+					]
+				]
+			]
+		]
+
+
+		// output log
+		+SVerticalBox::Slot()
+		.AutoHeight()
+		.Padding(0,4)
+		[
+			SNew(SBorder)
+			.BorderImage(FAppStyle::GetBrush("Brushes.Header"))
+			[
+				SNew(SHorizontalBox)
+
+				// output log title
+				+SHorizontalBox::Slot()
+				.FillWidth(1)
+				.Padding(4,2)
+				.VAlign(VAlign_Center)
+				[
+					SNew(STextBlock)				
+					.Text(LOCTEXT("OutputLogAreaTitle", "Output Log"))
+				]
+
+				// filter button
+				+SHorizontalBox::Slot()
+				.AutoWidth()
+				.Padding(4,2)
+				.VAlign(VAlign_Center)
+				[
+					OutputLog->CreateFilterWidget()
+				]
+			]
+		]
+
+		// main output log
+		+SVerticalBox::Slot()
+		.FillHeight(1)
+		[
+			OutputLog.ToSharedRef()
+		]
+	];
+}
+END_SLATE_FUNCTION_BUILD_OPTIMIZATION
+
+
+
+void SCustomLaunchLaunchPanel::SetLauncherWorker( const ILauncherWorkerRef& Worker )
+{
+	LauncherWorker = Worker;
+
+	Worker->GetTasks(TaskList);
+
+	if (SubwayProgressBar.IsValid())
+	{
+		SubwayProgressBar->ClearChildren();
+		for (ILauncherTaskPtr Task : TaskList)
+		{
+			SubwayProgressBar->AddSlot()
+			.Image(FCustomLaunchUIStyle::GetBrushForTask(Task))
+			.State(this, &SCustomLaunchLaunchPanel::GetSegmentedProgressBarState, Task)
+			.ToolTipText(FText::FromString(Task->GetDesc()))
+			;
+		}
+	}
+
+
+	Model->ClearLogMessages();
+	LaunchLogTextMarshaller->RefreshAllLogMessages();
+	OutputLog->RefreshLog();
+
+	Worker->OnOutputReceived().AddRaw(this, &SCustomLaunchLaunchPanel::HandleOutputReceived);
+}
+
+
+void SCustomLaunchLaunchPanel::Tick( const FGeometry& AllottedGeometry, const double InCurrentTime, const float InDeltaTime )
+{
+	LaunchLogTextMarshaller->FlushPendingLogMessages();
+	OutputLog->RequestForceScroll(true);
+}
+
+
+void SCustomLaunchLaunchPanel::HandleOutputReceived(const FString& InMessage)
+{
+	ELogVerbosity::Type Verbosity = ELogVerbosity::Log;
+
+	if (InMessage.StartsWith(TEXT("Parsing command line:"), ESearchCase::CaseSensitive))
+	{
+		Verbosity = ELogVerbosity::Display;
+	}
+	else if ( InMessage.Contains(TEXT("Error:"), ESearchCase::IgnoreCase) )
+	{
+		Verbosity = ELogVerbosity::Error;
+	}
+	else if ( InMessage.Contains(TEXT("Warning:"), ESearchCase::IgnoreCase) )
+	{
+		Verbosity = ELogVerbosity::Warning;
+	}
+
+	
+	
+	TSharedPtr<CustomLaunchUI::FLaunchLogMessage> Message = Model->AddLogMessage(InMessage, Verbosity);
+	LaunchLogTextMarshaller->AddPendingLogMessage(Message);
+}
+
+
+bool SCustomLaunchLaunchPanel::IsBuilding() const
+{
+	ILauncherWorkerPtr LauncherWorkerPtr = LauncherWorker.Pin();
+	if (LauncherWorkerPtr != nullptr)
+	{
+		return LauncherWorkerPtr->GetStatus() == ELauncherWorkerStatus::Busy;
+	}
+
+	return false;
+}
+
+bool SCustomLaunchLaunchPanel::IsIdle() const
+{
+	ILauncherWorkerPtr LauncherWorkerPtr = LauncherWorker.Pin();
+	if (LauncherWorkerPtr != nullptr)
+	{
+		return LauncherWorkerPtr->GetStatus() == ELauncherWorkerStatus::Canceled || LauncherWorkerPtr->GetStatus() == ELauncherWorkerStatus::Completed;
+	}
+
+	return false;
+}
+
+const FSlateBrush* SCustomLaunchLaunchPanel::GetProfileImage() const
+{
+	ILauncherProfilePtr LauncherProfile = GetLauncherProfile();
+	if (LauncherProfile != nullptr)
+	{
+		const PlatformInfo::FTargetPlatformInfo* PlatformInfo = Model->GetPlatformInfo(LauncherProfile);
+		return FCustomLaunchUIStyle::GetProfileBrushForPlatform(PlatformInfo, EPlatformIconSize::Large);
+	}
+
+	return FStyleDefaults::GetNoBrush();
+}
+
+
+FText SCustomLaunchLaunchPanel::GetProfileName() const
+{
+	ILauncherProfilePtr LauncherProfile = GetLauncherProfile();
+	if (LauncherProfile != nullptr)
+	{
+		return FText::FromString(LauncherProfile->GetName());
+	}
+
+	return FText::GetEmpty();
+}
+
+
+FText SCustomLaunchLaunchPanel::GetProfileDescription() const
+{
+	ILauncherProfilePtr LauncherProfile = GetLauncherProfile();
+	if (LauncherProfile != nullptr)
+	{
+		if (Model->IsBasicLaunchProfile(LauncherProfile))
+		{
+			const TSharedPtr<ITargetDeviceProxy> DeviceProxy = Model->GetDeviceProxy(LauncherProfile.ToSharedRef());
+			if (DeviceProxy.IsValid())
+			{
+				return FText::FromString(DeviceProxy->GetName());
+			}
+		}
+		else
+		{
+			return FText::FromString(LauncherProfile->GetDescription());
+		}
+	}
+
+	return FText::GetEmpty();
+}
+
+
+FText SCustomLaunchLaunchPanel::GetProfileProjectName() const
+{
+	ILauncherProfilePtr LauncherProfile = GetLauncherProfile();
+	if (LauncherProfile != nullptr)
+	{
+		return FText::FromString(LauncherProfile->GetProjectName());
+	}
+
+	return FText::GetEmpty();
+}
+FText SCustomLaunchLaunchPanel::GetProfileConfigurationName() const
+{
+	ILauncherProfilePtr LauncherProfile = GetLauncherProfile();
+	if (LauncherProfile != nullptr)
+	{
+		return EBuildConfigurations::ToText(LauncherProfile->GetBuildConfiguration());
+	}
+
+	return FText::GetEmpty();
+}
+FText SCustomLaunchLaunchPanel::GetProfileTargetName() const
+{
+	ILauncherProfilePtr LauncherProfile = GetLauncherProfile();
+	if (LauncherProfile != nullptr)
+	{
+		FString BuildTarget = LauncherProfile->GetBuildTarget();
+		if (BuildTarget.IsEmpty())
+		{
+			return LOCTEXT("DefaultTargetName", "Target Default");
+		}
+		else
+		{
+			return FText::FromString(BuildTarget);
+		}
+	}
+
+	return FText::GetEmpty();
+}
+
+
+FText SCustomLaunchLaunchPanel::GetProfileContentSchemeName() const
+{
+	ILauncherProfilePtr LauncherProfile = GetLauncherProfile();
+	if (LauncherProfile != nullptr)
+	{
+		if (Model->IsAdvancedProfile(LauncherProfile.ToSharedRef()))
+		{
+			return LOCTEXT("AdvancedProfileInfo", "Advanced/Legacy Profile");
+		}
+
+		CustomLaunchUI::EContentScheme ContentScheme = Model->DetermineProfileContentScheme(LauncherProfile.ToSharedRef());
+		return CustomLaunchUI::GetContentSchemeDisplayName(ContentScheme);
+	}
+
+	return FText::GetEmpty();
+}
+
+
+FText SCustomLaunchLaunchPanel::GetCurrentTaskDescription() const
+{
+	if (TaskList.Num() > 0 && LauncherWorker.IsValid())
+	{
+		if (TaskList.Last()->GetStatus() == ELauncherTaskStatus::Completed)
+		{
+			return LOCTEXT("TaskStatusAllCompleted", "Completed All Tasks");
+		}
+
+		ILauncherTaskPtr LatestLauncherTask = TaskList[0];
+		for (const ILauncherTaskPtr& LauncherTask : TaskList)
+		{
+			if (LauncherTask->GetStatus() == ELauncherTaskStatus::Canceled || LauncherTask->GetStatus() == ELauncherTaskStatus::Pending)
+			{
+				break;
+			}
+			LatestLauncherTask = LauncherTask;
+		}
+
+		return FText::FromString(LatestLauncherTask->GetDesc());
+	}
+
+	return FText::GetEmpty();
+}
+
+FText SCustomLaunchLaunchPanel::GetTotalDurationDescription() const
+{
+	if (TaskList.Num() > 0 && LauncherWorker.IsValid())
+	{
+		FTimespan Duration;
+		for (const ILauncherTaskPtr& LauncherTask : TaskList)
+		{
+			Duration += LauncherTask->GetDuration();
+		}
+
+		return FText::AsTimespan(Duration);
+	}
+
+	return FText::GetEmpty();
+}
+
+FText SCustomLaunchLaunchPanel::GetProgressDescription() const
+{
+	if (TaskList.Num() > 0 && LauncherWorker.IsValid())
+	{
+		TOptional<float> ProgressPercent = GetProgressPercent();
+
+		for (const ILauncherTaskPtr& LauncherTask : TaskList)
+		{
+			if (LauncherTask->GetStatus() == ELauncherTaskStatus::Busy && ProgressPercent.IsSet())
+			{
+				return FText::AsPercent(ProgressPercent.GetValue());
+			}
+			else if (LauncherTask->GetStatus() == ELauncherTaskStatus::Failed)
+			{
+				return LOCTEXT("TaskStatusFailed", "Failed");
+			}
+		}
+
+		switch (TaskList.Last()->GetStatus())
+		{
+			case ELauncherTaskStatus::Canceled: return LOCTEXT("TaskStatusCanceled", "Canceled");
+			//case ELauncherTaskStatus::Completed: return LOCTEXT("TaskStatusCompleted", "Completed");
+		}
+	}
+
+	return FText::GetEmpty();
+}
+
+TOptional<float> SCustomLaunchLaunchPanel::GetProgressPercent() const
+{
+	if (TaskList.Num() > 0 && LauncherWorker.IsValid())
+	{
+		int32 NumFinished = 0;
+		for (int32 TaskIndex = 0; TaskIndex < TaskList.Num(); ++TaskIndex)
+		{
+			if (TaskList[TaskIndex]->IsFinished())
+			{
+				++NumFinished;
+			}
+		}
+
+		return ((float)NumFinished / TaskList.Num());
+	}
+
+	TOptional<float> Result;
+	return Result;
+}
+
+SSegmentedProgressBar::EState SCustomLaunchLaunchPanel::GetSegmentedProgressBarState( ILauncherTaskPtr Task ) const
+{
+	if (Task.IsValid())
+	{
+		switch (Task->GetStatus())
+		{
+			case ELauncherTaskStatus::Busy:      return SSegmentedProgressBar::EState::Pending;
+			case ELauncherTaskStatus::Canceled:  return SSegmentedProgressBar::EState::Canceled;
+			case ELauncherTaskStatus::Completed: return SSegmentedProgressBar::EState::Completed;
+			case ELauncherTaskStatus::Failed:    return SSegmentedProgressBar::EState::Failed;
+			case ELauncherTaskStatus::Pending:   return SSegmentedProgressBar::EState::Pending;
+		}
+	}
+
+	return SSegmentedProgressBar::EState::None;
+}
+
+FReply SCustomLaunchLaunchPanel::OnCancelButtonClicked()
+{
+	ILauncherWorkerPtr LauncherWorkerPtr = LauncherWorker.Pin();
+	if (LauncherWorkerPtr != nullptr && LauncherWorkerPtr->GetStatus() == ELauncherWorkerStatus::Busy)
+	{
+		LauncherWorkerPtr->Cancel();
+	}
+
+	return FReply::Handled();
+}
+
+
+FReply SCustomLaunchLaunchPanel::OnRetryButtonClicked()
+{
+	check(IsIdle());
+	if (OnRerunClicked.IsBound())
+	{
+		OnRerunClicked.Execute();
+	}
+	return FReply::Handled();
+}
+
+
+FReply SCustomLaunchLaunchPanel::OnDoneButtonClicked()
+{
+	check(IsIdle());
+	if (OnCloseClicked.IsBound())
+	{
+		OnCloseClicked.Execute();
+	}
+	return FReply::Handled();
+}
+
+ILauncherProfilePtr SCustomLaunchLaunchPanel::GetLauncherProfile() const
+{
+	ILauncherWorkerPtr LauncherWorkerPtr = LauncherWorker.Pin();
+	return (LauncherWorkerPtr != nullptr) ? LauncherWorkerPtr->GetLauncherProfile() : nullptr;
+}
+
+
+
+#undef LOCTEXT_NAMESPACE
diff --git a/Engine/Plugins/Developer/CustomLaunchUI/Source/Private/Widgets/Output/SCustomLaunchLaunchPanel.h b/Engine/Plugins/Developer/CustomLaunchUI/Source/Private/Widgets/Output/SCustomLaunchLaunchPanel.h
new file mode 100644
--- /dev/null
+++ b/Engine/Plugins/Developer/CustomLaunchUI/Source/Private/Widgets/Output/SCustomLaunchLaunchPanel.h
@@ -0,0 +1,70 @@
+// Copyright Epic Games, Inc. All Rights Reserved.
+
+#pragma once
+
+#include "Internationalization/Text.h"
+#include "Framework/SlateDelegates.h"
+#include "Widgets/SCompoundWidget.h"
+#include "Widgets/SSegmentedProgressBar.h"
+#include "Widgets/Output/CustomLaunchOutputLogMarshaller.h"
+#include "Model/CustomLaunchUIModel.h"
+#include "ILauncherWorker.h"
+
+class SWidget;
+class SCustomLaunchOutputLog;
+class ITableRow;
+class STableViewBase;
+template<typename T> class SListView;
+
+
+class SCustomLaunchLaunchPanel
+	: public SCompoundWidget
+{
+public:
+	SLATE_BEGIN_ARGS(SCustomLaunchLaunchPanel) {}
+		SLATE_EVENT(FOnClicked, OnCloseClicked)
+		SLATE_EVENT(FOnClicked, OnRerunClicked)
+	SLATE_END_ARGS()
+
+public:
+	void Construct(const FArguments& InArgs, const TSharedRef<CustomLaunchUI::FModel>& InModel);
+	void Tick( const FGeometry& AllottedGeometry, const double InCurrentTime, const float InDeltaTime ) override;
+
+
+	void SetLauncherWorker( const ILauncherWorkerRef& Worker );
+
+private:
+	bool IsBuilding() const;
+	bool IsIdle() const;
+	const FSlateBrush* GetProfileImage() const;
+	FText GetProfileName() const;
+	FText GetProfileDescription() const;
+	FText GetProfileProjectName() const;
+	FText GetProfileConfigurationName() const;
+	FText GetProfileTargetName() const;
+	FText GetProfileContentSchemeName() const;	
+	FText GetCurrentTaskDescription() const;
+	FText GetProgressDescription() const;
+	FText GetTotalDurationDescription() const;
+	TOptional<float> GetProgressPercent() const;
+	SSegmentedProgressBar::EState GetSegmentedProgressBarState(ILauncherTaskPtr Task) const;
+	FReply OnCancelButtonClicked();
+	FReply OnRetryButtonClicked();
+	FReply OnDoneButtonClicked();
+
+	ILauncherProfilePtr GetLauncherProfile() const;
+
+	void HandleOutputReceived(const FString& InMessage);
+
+	FOnClicked OnCloseClicked;
+	FOnClicked OnRerunClicked;
+
+	TWeakPtr<ILauncherWorker> LauncherWorker;
+	TArray<ILauncherTaskPtr> TaskList;
+
+	TSharedPtr<CustomLaunchUI::FLaunchLogTextLayoutMarshaller> LaunchLogTextMarshaller;
+	TSharedPtr<SCustomLaunchOutputLog> OutputLog;
+	TSharedPtr<SSegmentedProgressBar> SubwayProgressBar;
+
+	TSharedPtr<CustomLaunchUI::FModel> Model;
+};
diff --git a/Engine/Plugins/Developer/CustomLaunchUI/Source/Private/Widgets/Output/SCustomLaunchOutputLog.cpp b/Engine/Plugins/Developer/CustomLaunchUI/Source/Private/Widgets/Output/SCustomLaunchOutputLog.cpp
new file mode 100644
--- /dev/null
+++ b/Engine/Plugins/Developer/CustomLaunchUI/Source/Private/Widgets/Output/SCustomLaunchOutputLog.cpp
@@ -0,0 +1,484 @@
+// Copyright Epic Games, Inc. All Rights Reserved.
+
+#include "Widgets/Output/SCustomLaunchOutputLog.h"
+
+#include "Styling/AppStyle.h"
+#include "Styling/CustomLaunchUIStyle.h"
+#include "Framework/Commands/UIAction.h"
+#include "Framework/MultiBox/MultiBoxBuilder.h"
+#include "Framework/Application/SlateApplication.h"
+#include "SlateOptMacros.h"
+#include "Widgets/SBoxPanel.h"
+#include "Widgets/Layout/SSpacer.h"
+#include "Widgets/Input/SComboButton.h"
+#include "Widgets/Input/SButton.h"
+#include "Widgets/Input/SSearchBox.h"
+#include "Widgets/Input/SMultiLineEditableTextBox.h"
+#include "Widgets/Text/STextBlock.h"
+#include "Widgets/Images/SImage.h"
+#include "HAL/PlatformApplicationMisc.h"
+#include "Misc/StringBuilder.h"
+#include "Misc/Paths.h"
+#include "Misc/FileHelper.h"
+#include "Misc/MessageDialog.h"
+#include "DesktopPlatformModule.h"
+
+#define LOCTEXT_NAMESPACE "SCustomLaunchOutputLog"
+
+
+BEGIN_SLATE_FUNCTION_BUILD_OPTIMIZATION
+void SCustomLaunchOutputLog::Construct(const FArguments& InArgs, const TSharedRef<CustomLaunchUI::FModel>& InModel, const TSharedRef<CustomLaunchUI::FLaunchLogTextLayoutMarshaller>& InLaunchLogTextMarshaller)
+{
+	Model = InModel;
+	LaunchLogTextMarshaller = InLaunchLogTextMarshaller;
+
+	LogMessageTextBox = SNew(SMultiLineEditableTextBox)
+		.Style(FAppStyle::Get(), "Log.TextBox")
+		.Marshaller(LaunchLogTextMarshaller)
+		.IsReadOnly(true)
+		.AlwaysShowScrollbars(true)
+		.AutoWrapText_Lambda( [this] () { return bWordWrap;} )
+		.OnVScrollBarUserScrolled(this, &SCustomLaunchOutputLog::OnUserScrolled)
+		.ContextMenuExtender(this, &SCustomLaunchOutputLog::ExtendTextBoxMenu);
+		;
+
+	ChildSlot
+	[
+		SNew(SVerticalBox)
+
+		+SVerticalBox::Slot()
+		.FillHeight(1)
+		.Padding(0)
+		[
+			SNew(SHorizontalBox)
+
+			+SHorizontalBox::Slot()
+			.FillWidth(1)
+			[
+				LogMessageTextBox.ToSharedRef()
+			]
+		]
+
+		+SVerticalBox::Slot()
+		.AutoHeight()
+		[
+			SNew(SBorder)
+			.BorderImage(FAppStyle::GetBrush("Brushes.Header"))
+			[
+				SNew(SHorizontalBox)
+
+				// padding
+				+SHorizontalBox::Slot()
+				.FillWidth(1)
+				[
+					SNew(SSpacer)
+				]
+
+				// copy button
+				+SHorizontalBox::Slot()
+				.AutoWidth()
+				.HAlign(HAlign_Center)
+				.VAlign(VAlign_Center)
+				.Padding(4,2)
+				[
+					SAssignNew(CopyButton, SButton)
+					.OnClicked(this, &SCustomLaunchOutputLog::OnCopyClicked)
+					.IsEnabled_Lambda([InModel]() { return (InModel->GetNumLogMessages() > 0); } )
+					.ToolTipText(LOCTEXT("CopyButtonTip", "Copy entire log to the clipboard"))
+					[
+						SNew(SImage)
+						.Image(FAppStyle::Get().GetBrush("GenericCommands.Copy"))
+						.ColorAndOpacity(FSlateColor::UseSubduedForeground())
+					]
+				]
+
+				// clear button
+				+SHorizontalBox::Slot()
+				.AutoWidth()
+				.HAlign(HAlign_Center)
+				.VAlign(VAlign_Center)
+				.Padding(4,2)
+				[
+					SAssignNew(ClearButton, SButton)
+					.OnClicked(this, &SCustomLaunchOutputLog::OnClearClicked)
+					.IsEnabled_Lambda([InModel]() { return (InModel->GetNumLogMessages() > 0); } )
+					.ToolTipText(LOCTEXT("ClearButtonTip", "Clear all messages"))
+					[
+						SNew(SImage)
+						.Image(FCustomLaunchUIStyle::Get().GetBrush("Icons.ClearLog"))
+						.ColorAndOpacity(FSlateColor::UseSubduedForeground())
+					]
+				]
+
+				// save button
+				+SHorizontalBox::Slot()
+				.AutoWidth()
+				.HAlign(HAlign_Center)
+				.VAlign(VAlign_Center)
+				.Padding(4,2)
+				[
+					SAssignNew(SaveButton, SButton)
+					.OnClicked(this, &SCustomLaunchOutputLog::OnSaveClicked)
+					.Visibility( FDesktopPlatformModule::Get() != nullptr ? EVisibility::Visible : EVisibility::Collapsed )
+					.IsEnabled_Lambda([InModel]() { return (InModel->GetNumLogMessages() > 0); } )
+					.ToolTipText(LOCTEXT("SaveButtonTip", "Save log to file"))
+					[
+						SNew(SImage)
+						.Image(FAppStyle::Get().GetBrush("Icons.Save"))
+						.ColorAndOpacity(FSlateColor::UseSubduedForeground())
+					]
+				]
+			]
+		]
+	];
+
+	bIsUserScrolled = false;
+	RequestForceScroll();
+}
+END_SLATE_FUNCTION_BUILD_OPTIMIZATION
+
+
+
+BEGIN_SLATE_FUNCTION_BUILD_OPTIMIZATION
+TSharedRef<SWidget> SCustomLaunchOutputLog::CreateFilterWidget()
+{
+	return SNew(SHorizontalBox)
+	.Visibility_Lambda( [this]() { return GetParentWidget()->GetVisibility(); } )
+
+	// search/filter box
+	+SHorizontalBox::Slot()
+	.AutoWidth()
+	[
+		SNew(SSearchBox)
+			.OnTextChanged(this, &SCustomLaunchOutputLog::OnFilterTextChanged)
+			.OnTextCommitted(this, &SCustomLaunchOutputLog::OnFilterTextCommitted)
+			.HintText(LOCTEXT("FilterTextHint", "Text Filter"))
+			.InitialText(CurrentFilterText)
+			.MinDesiredWidth(128)
+			.ToolTipText(LOCTEXT("FilterTextTip", "Only show lines that contain the specified search text") )
+			.DelayChangeNotificationsWhileTyping(true)
+	]
+
+	// filter button
+	+SHorizontalBox::Slot()
+	.AutoWidth()
+	[
+		SNew(SComboButton)
+		.ComboButtonStyle(FAppStyle::Get(), "SimpleComboButton")
+		.ToolTipText(LOCTEXT("AddFilterToolTip", "Add an output log filter."))
+		.OnGetMenuContent(this, &SCustomLaunchOutputLog::MakeFilterMenu)
+		.MenuPlacement(MenuPlacement_BelowRightAnchor)
+		.ButtonContent()
+		[
+			SNew(SHorizontalBox)
+			+SHorizontalBox::Slot()
+			.AutoWidth()
+			[
+				SNew(SImage)
+				.Image(FAppStyle::Get().GetBrush("Icons.Filter"))
+				.ColorAndOpacity(FSlateColor::UseForeground())
+			]
+			+SHorizontalBox::Slot()
+			.AutoWidth()
+			.Padding(2, 0, 0, 0)
+			[
+				SNew(STextBlock)
+				.Text(LOCTEXT("Filters", "Filters"))
+				.ColorAndOpacity(FSlateColor::UseForeground())
+			]
+		]
+	]
+	;
+}
+END_SLATE_FUNCTION_BUILD_OPTIMIZATION
+
+
+
+TSharedRef<SWidget> SCustomLaunchOutputLog::MakeFilterMenu()
+{
+	const bool bShouldCloseWindowAfterMenuSelection = true;
+	const bool bCloseSelfOnly = false;
+	const bool bSearchable = false;
+	FMenuBuilder MenuBuilder(bShouldCloseWindowAfterMenuSelection, nullptr, nullptr, bCloseSelfOnly, &FCoreStyle::Get(), bSearchable );
+
+	MenuBuilder.AddMenuEntry(
+		LOCTEXT("FilterAllLabel", "All Messages"),
+		LOCTEXT("FilterAllTip", "Show all messages"),
+		FSlateIcon(),
+		FUIAction(
+			FExecuteAction::CreateSP(this, &SCustomLaunchOutputLog::OnFilterChanged, CustomLaunchUI::ELogFilter::All),
+			FCanExecuteAction(),
+			FGetActionCheckState::CreateLambda( [this]() { return (LaunchLogTextMarshaller->GetFilter() == CustomLaunchUI::ELogFilter::All) ? ECheckBoxState::Checked : ECheckBoxState::Unchecked; })
+		),
+		NAME_None,
+		EUserInterfaceActionType::RadioButton);
+
+	MenuBuilder.AddMenuEntry(
+		LOCTEXT("FilterWarningsAndErrorsLabel", "Warnings & Errors"),
+		LOCTEXT("FilterWarningsAndErrorsTip", "Show only warnings and errors"),
+		FSlateIcon(),
+		FUIAction(
+			FExecuteAction::CreateSP(this, &SCustomLaunchOutputLog::OnFilterChanged, CustomLaunchUI::ELogFilter::WarningsAndErrors),
+			FCanExecuteAction(),
+			FGetActionCheckState::CreateLambda( [this]() { return (LaunchLogTextMarshaller->GetFilter() == CustomLaunchUI::ELogFilter::WarningsAndErrors) ? ECheckBoxState::Checked : ECheckBoxState::Unchecked; })
+		),
+		NAME_None,
+		EUserInterfaceActionType::RadioButton);
+
+	MenuBuilder.AddMenuEntry(
+		LOCTEXT("FilterErrorsLabel", "Errors"),
+		LOCTEXT("FilterErrorsTip", "Show only errors"),
+		FSlateIcon(),
+		FUIAction(
+			FExecuteAction::CreateSP(this, &SCustomLaunchOutputLog::OnFilterChanged, CustomLaunchUI::ELogFilter::Errors),
+			FCanExecuteAction(),
+			FGetActionCheckState::CreateLambda( [this]() { return (LaunchLogTextMarshaller->GetFilter() == CustomLaunchUI::ELogFilter::Errors) ? ECheckBoxState::Checked : ECheckBoxState::Unchecked; })
+		),
+		NAME_None,
+		EUserInterfaceActionType::RadioButton);
+
+	MenuBuilder.AddMenuEntry(
+		LOCTEXT("SearchFiltersLogLabel", "Only Matching Lines"),
+		LOCTEXT("SearchFiltersLogTip", "Only show log lines that match the text search filter"),
+		FSlateIcon(),
+		FUIAction(
+			FExecuteAction::CreateSP(this, &SCustomLaunchOutputLog::OnSearchBoxFiltersLogToggle ),
+			FCanExecuteAction(),
+			FGetActionCheckState::CreateLambda( [this]() { return bSearchBoxFiltersLog ? ECheckBoxState::Checked : ECheckBoxState::Unchecked; } )
+		),
+		NAME_None,
+		EUserInterfaceActionType::Check);
+
+	MenuBuilder.AddMenuSeparator();
+
+	MenuBuilder.AddMenuEntry(
+		LOCTEXT("WordWrapLabel", "Enable Word Wrapping"),
+		LOCTEXT("WordWrapTip", "Split long log entries across multiple lines"),
+		FSlateIcon(),
+		FUIAction(
+			FExecuteAction::CreateSP(this, &SCustomLaunchOutputLog::OnWordWrapToggle ),
+			FCanExecuteAction(),
+			FGetActionCheckState::CreateLambda( [this]() { return bWordWrap ? ECheckBoxState::Checked : ECheckBoxState::Unchecked; } )
+		),
+		NAME_None,
+		EUserInterfaceActionType::Check);
+
+	return MenuBuilder.MakeWidget();
+}
+
+
+void SCustomLaunchOutputLog::ExtendTextBoxMenu(FMenuBuilder& MenuBuilder)
+{
+	MenuBuilder.BeginSection(NAME_None, LOCTEXT("LogMenuSectionLabel", "Output Log"));
+
+	MenuBuilder.AddMenuEntry(
+		LOCTEXT("ClearLogLabel", "Clear"),
+		LOCTEXT("ClearLogTip", "Clears all log messages"),
+		FSlateIcon(FCustomLaunchUIStyle::GetStyleSetName(), "Icons.ClearLog"),
+		FUIAction(
+			FExecuteAction::CreateSP(this, &SCustomLaunchOutputLog::ClearLog ),
+			FCanExecuteAction::CreateLambda([this]() { return Model->GetNumLogMessages() > 0; } ),
+			FGetActionCheckState()
+		),
+		NAME_None,
+		EUserInterfaceActionType::Check);
+
+
+	MenuBuilder.AddMenuEntry(
+		LOCTEXT("SageLogLabel", "Save As..."),
+		LOCTEXT("SageLogTip", "Save log to file"),
+		FSlateIcon(FAppStyle::GetAppStyleSetName(), "Icons.Save"),
+		FUIAction(
+			FExecuteAction::CreateSP(this, &SCustomLaunchOutputLog::SaveLog ),
+			FCanExecuteAction::CreateLambda([this]() { return Model->GetNumLogMessages() > 0; } ),
+			FGetActionCheckState()
+		));
+
+	MenuBuilder.EndSection();
+}
+
+
+
+void SCustomLaunchOutputLog::OnFilterTextChanged(const FText& FilterText)
+{
+	CurrentFilterText = FilterText;
+
+	// only show matching lines in the log
+	if (bSearchBoxFiltersLog)
+	{
+		LogMessageTextBox->GoTo(FTextLocation(0));
+		LaunchLogTextMarshaller->SetFilterString(FilterText.ToString());
+		LogMessageTextBox->Refresh();
+		RequestForceScroll();
+	}
+
+	// highlight the first item
+	LogMessageTextBox->BeginSearch(FilterText);
+}
+
+void SCustomLaunchOutputLog::OnFilterTextCommitted(const FText& InFilterText, ETextCommit::Type InCommitType)
+{
+	bool bFilterChanged = !InFilterText.EqualTo(CurrentFilterText);
+	if (bFilterChanged)
+	{
+		OnFilterTextChanged(InFilterText);
+	}
+	else if (InCommitType == ETextCommit::OnEnter)
+	{
+		bool bReverse = FSlateApplication::Get().GetModifierKeys().IsShiftDown();
+		bIsUserScrolled = true;
+		LogMessageTextBox->AdvanceSearch(bReverse);
+	}
+}
+
+void SCustomLaunchOutputLog::OnFilterChanged(CustomLaunchUI::ELogFilter Filter)
+{
+	LogMessageTextBox->GoTo(FTextLocation(0));
+	LaunchLogTextMarshaller->SetFilter(Filter);
+	LogMessageTextBox->Refresh();
+	RequestForceScroll();
+}
+
+void SCustomLaunchOutputLog::OnWordWrapToggle()
+{
+	bWordWrap = !bWordWrap;
+	RequestForceScroll(true);
+}
+
+void SCustomLaunchOutputLog::OnSearchBoxFiltersLogToggle()
+{
+	bSearchBoxFiltersLog = !bSearchBoxFiltersLog;
+	if (!bSearchBoxFiltersLog)
+	{
+		LaunchLogTextMarshaller->SetFilterString(FString());
+	}
+	else
+	{
+		LaunchLogTextMarshaller->SetFilterString(CurrentFilterText.ToString());
+	}
+	RefreshLog();
+}
+
+void SCustomLaunchOutputLog::OnUserScrolled(float ScrollOffset)
+{
+	bIsUserScrolled = ScrollOffset < 1.0 && !FMath::IsNearlyEqual(ScrollOffset, 1.0f);
+}
+
+
+
+void SCustomLaunchOutputLog::RefreshLog()
+{
+	LogMessageTextBox->GoTo(FTextLocation(0));
+	LaunchLogTextMarshaller->RefreshAllLogMessages();
+	LogMessageTextBox->Refresh();
+	RequestForceScroll();
+}
+
+void SCustomLaunchOutputLog::RequestForceScroll(bool bIfUserHasNotScrolledUp)
+{
+	if (LaunchLogTextMarshaller->GetNumFilteredMessages() > 0
+		&& (!bIfUserHasNotScrolledUp || !bIsUserScrolled))
+	{
+		LogMessageTextBox->ScrollTo(ETextLocation::EndOfDocument);
+		bIsUserScrolled = false;
+	}
+}
+
+
+
+
+FReply SCustomLaunchOutputLog::OnSaveClicked()
+{
+	SaveLog();
+	return FReply::Handled();
+}
+
+
+
+FReply SCustomLaunchOutputLog::OnClearClicked()
+{
+	ClearLog();
+	return FReply::Handled();
+}
+
+
+
+FReply SCustomLaunchOutputLog::OnCopyClicked()
+{
+	CopyLog();
+	return FReply::Handled();
+}
+
+
+FString SCustomLaunchOutputLog::GetLogAsString( bool bSelectedLinesOnly ) const
+{
+	if (bSelectedLinesOnly)
+	{
+		return LogMessageTextBox->GetSelectedText().ToString();
+	}
+	else
+	{
+		return LogMessageTextBox->GetPlainText().ToString();
+	}
+}
+
+
+
+
+void SCustomLaunchOutputLog::ClearLog()
+{
+	Model->ClearLogMessages();
+	RefreshLog();
+	bIsUserScrolled = false;
+}
+
+void SCustomLaunchOutputLog::CopyLog()
+{
+	FString LogText = GetLogAsString(false);
+	if (!LogText.IsEmpty())
+	{
+		FPlatformApplicationMisc::ClipboardCopy( *LogText );
+	}
+}
+
+
+
+void SCustomLaunchOutputLog::SaveLog()
+{
+	static FString LastSavePath;
+
+	// @todo: try using ISlateFileDialogsModule
+	IDesktopPlatform* DesktopPlatform = FDesktopPlatformModule::Get();
+	if (DesktopPlatform != nullptr)
+	{
+		TArray<FString> FileNames;
+
+		if (DesktopPlatform->SaveFileDialog(
+			FSlateApplication::Get().FindBestParentWindowHandleForDialogs(nullptr),
+			LOCTEXT("SaveLogDialogTitle", "Save Log As...").ToString(),
+			*LastSavePath,
+			TEXT("BuidCookRun.log"),
+			TEXT("Log Files (*.log)|*.log"),
+			EFileDialogFlags::None,
+			FileNames))
+		{
+			if (FileNames.Num() > 0)
+			{
+				FString FileName = FileNames[0];
+				LastSavePath = FPaths::GetPath(FileName); // record last used directory
+				if (FPaths::GetExtension(FileName).IsEmpty()) // make sure there's an extension
+				{
+					FileName += TEXT(".log");
+				}
+
+				if (!FFileHelper::SaveStringToFile( GetLogAsString(false), *FileName))
+				{
+					FMessageDialog::Open(EAppMsgType::Ok, LOCTEXT("SaveLogDialogFail", "Failed to save the log"));
+				}
+			}
+		}
+	}
+}
+
+
+#undef LOCTEXT_NAMESPACE
diff --git a/Engine/Plugins/Developer/CustomLaunchUI/Source/Private/Widgets/Output/SCustomLaunchOutputLog.h b/Engine/Plugins/Developer/CustomLaunchUI/Source/Private/Widgets/Output/SCustomLaunchOutputLog.h
new file mode 100644
--- /dev/null
+++ b/Engine/Plugins/Developer/CustomLaunchUI/Source/Private/Widgets/Output/SCustomLaunchOutputLog.h
@@ -0,0 +1,62 @@
+// Copyright Epic Games, Inc. All Rights Reserved.
+
+#pragma once
+
+#include "Internationalization/Text.h"
+#include "Framework/SlateDelegates.h"
+#include "Widgets/SCompoundWidget.h"
+#include "Model/CustomLaunchUIModel.h"
+#include "Widgets/Output/CustomLaunchOutputLogMarshaller.h"
+
+class SWidget;
+class SButton;
+class SMultiLineEditableTextBox;
+
+class SCustomLaunchOutputLog 
+	: public SCompoundWidget
+{
+public:
+	SLATE_BEGIN_ARGS(SCustomLaunchOutputLog) {}
+	SLATE_END_ARGS()
+
+public:
+	void Construct(const FArguments& InArgs, const TSharedRef<CustomLaunchUI::FModel>& InModel, const TSharedRef<CustomLaunchUI::FLaunchLogTextLayoutMarshaller>& InLaunchLogTextMarshaller);
+
+	void CopyLog();
+	void SaveLog();
+	void ClearLog();
+	FString GetLogAsString( bool bSelectedLinesOnly = false ) const;
+
+	void RefreshLog();
+
+	void RequestForceScroll( bool bIfUserHasNotScrolledUp = false );
+
+	TSharedRef<SWidget> CreateFilterWidget();
+
+private:
+	TSharedRef<SWidget> MakeFilterMenu();
+	FReply OnSaveClicked();
+	FReply OnClearClicked();
+	FReply OnCopyClicked();
+
+	void OnFilterTextChanged(const FText& FilterText);
+	void OnFilterTextCommitted(const FText& InFilterText, ETextCommit::Type InCommitType);
+	void OnFilterChanged(CustomLaunchUI::ELogFilter Filter);
+	void OnWordWrapToggle();
+	void OnSearchBoxFiltersLogToggle();
+	void OnUserScrolled(float ScrollOffset);
+	void ExtendTextBoxMenu(class FMenuBuilder& Builder);
+
+	TSharedPtr<SMultiLineEditableTextBox> LogMessageTextBox;
+	TSharedPtr<SButton> SaveButton;
+	TSharedPtr<SButton> ClearButton;
+	TSharedPtr<SButton> CopyButton;
+
+private:
+	bool bIsUserScrolled = false;
+	bool bWordWrap = false;
+	bool bSearchBoxFiltersLog = true;
+	FText CurrentFilterText;
+	TSharedPtr<CustomLaunchUI::FLaunchLogTextLayoutMarshaller> LaunchLogTextMarshaller;
+	TSharedPtr<CustomLaunchUI::FModel> Model;
+};
diff --git a/Engine/Plugins/Developer/CustomLaunchUI/Source/Private/Widgets/Profiles/SCustomLaunchCustomProfileSelector.cpp b/Engine/Plugins/Developer/CustomLaunchUI/Source/Private/Widgets/Profiles/SCustomLaunchCustomProfileSelector.cpp
new file mode 100644
--- /dev/null
+++ b/Engine/Plugins/Developer/CustomLaunchUI/Source/Private/Widgets/Profiles/SCustomLaunchCustomProfileSelector.cpp
@@ -0,0 +1,621 @@
+// Copyright Epic Games, Inc. All Rights Reserved.
+
+#include "Widgets/Profiles/SCustomLaunchCustomProfileSelector.h"
+
+#include "Styling/AppStyle.h"
+#include "Styling/StyleColors.h"
+#include "Styling/CustomLaunchUIStyle.h"
+#include "Framework/Commands/UIAction.h"
+#include "Framework/MultiBox/MultiBoxBuilder.h"
+#include "SlateOptMacros.h"
+#include "Widgets/SBoxPanel.h"
+#include "Widgets/Input/SHyperlink.h"
+#include "Widgets/Input/SComboBox.h"
+#include "Widgets/Input/SComboButton.h"
+#include "Widgets/Layout/SSplitter.h"
+#include "Widgets/Layout/SScrollBorder.h"
+#include "Widgets/Text/STextBlock.h"
+#include "Widgets/Text/SInlineEditableTextBlock.h"
+#include "Widgets/Images/SImage.h"
+#include "Widgets/Views/SListView.h"
+#include "Widgets/SNullWidget.h"
+#include "Widgets/Shared/SCustomLaunchDeviceCombo.h"
+#include "SPositiveActionButton.h"
+#include "DesktopPlatformModule.h"
+#include "PlatformInfo.h"
+#include "ITargetDeviceProxy.h"
+#include "ITargetDeviceProxyManager.h"
+#include "Interfaces/ITargetPlatform.h"
+#include "Interfaces/ITargetPlatformManagerModule.h"
+#include "InstalledPlatformInfo.h"
+#include "Framework/Docking/TabManager.h"
+#include "Misc/ConfigCacheIni.h"
+#include "GameProjectHelper.h"
+
+
+#define LOCTEXT_NAMESPACE "SCustomLaunchCustomProfileSelector"
+
+
+SCustomLaunchCustomProfileSelector::~SCustomLaunchCustomProfileSelector()
+{
+	if (Model.IsValid())
+	{
+		Model->OnProfileSelected().RemoveAll(this);
+
+		Model->GetProfileManager()->OnProfileAdded().RemoveAll(this);
+		Model->GetProfileManager()->OnProfileRemoved().RemoveAll(this);
+
+		const TSharedRef<ITargetDeviceProxyManager>& DeviceProxyManager = Model->GetDeviceProxyManager();
+		DeviceProxyManager->OnProxyAdded().RemoveAll(this);
+		DeviceProxyManager->OnProxyRemoved().RemoveAll(this);
+	}
+}
+
+
+
+BEGIN_SLATE_FUNCTION_BUILD_OPTIMIZATION
+void SCustomLaunchCustomProfileSelector::Construct(const FArguments& InArgs, const TSharedRef<CustomLaunchUI::FModel>& InModel)
+{
+	Model = InModel;
+	OnProfileAdd = InArgs._OnProfileAdd;
+	OnProfileDelete = InArgs._OnProfileDelete;
+	OnProfileDuplicate = InArgs._OnProfileDuplicate;
+	OnProfileEdit = InArgs._OnProfileEdit;
+	OnProfileRename = InArgs._OnProfileRename;
+	OnProfileModified = InArgs._OnProfileModified;
+	ChangeProfileEditorVisibility = InArgs._ChangeProfileEditorVisibility;
+	EditPanelVisible = InArgs._EditPanelVisible;
+
+	CustomProfileListView = SNew(SListView<ILauncherProfilePtr>)
+	.ListItemsSource(&Model->GetAllProfiles())
+	.ClearSelectionOnClick(false)
+	.SelectionMode(ESelectionMode::Single)
+	.OnMouseButtonDoubleClick(this, &SCustomLaunchCustomProfileSelector::EditProfile)
+	.OnSelectionChanged(this, &SCustomLaunchCustomProfileSelector::OnSelectionChanged)
+	.OnGenerateRow(this, &SCustomLaunchCustomProfileSelector::GenerateCustomProfileRow)
+	.OnContextMenuOpening(this, &SCustomLaunchCustomProfileSelector::MakeContextMenu)
+	;
+
+	ChildSlot
+	[
+		// main content
+		SNew(SVerticalBox)
+		+SVerticalBox::Slot()
+		.FillHeight(1)
+		[
+			SNew(SScrollBorder, CustomProfileListView.ToSharedRef())
+			[
+				CustomProfileListView.ToSharedRef()
+			]
+		]
+	];
+
+	// hook profile changes
+	Model->GetProfileManager()->OnProfileAdded().AddSP(this, &SCustomLaunchCustomProfileSelector::OnCustomProfileAdded);
+	Model->GetProfileManager()->OnProfileRemoved().AddSP(this, &SCustomLaunchCustomProfileSelector::OnCustomProfileRemoved);
+
+	Model->OnProfileSelected().AddSP(this, &SCustomLaunchCustomProfileSelector::OnProfileSelected);
+}
+END_SLATE_FUNCTION_BUILD_OPTIMIZATION
+
+
+
+
+
+
+
+
+
+
+
+BEGIN_SLATE_FUNCTION_BUILD_OPTIMIZATION
+TSharedRef<ITableRow> SCustomLaunchCustomProfileSelector::GenerateCustomProfileRow(ILauncherProfilePtr Profile, const TSharedRef<STableViewBase>& OwnerTable)
+{
+	bool bIsBasicLaunchProfile = Model->IsBasicLaunchProfile(Profile);
+
+	TSharedPtr<SVerticalBox> DetailsBox;
+	TSharedPtr<SVerticalBox> ButtonsBox;
+	TSharedPtr<SHorizontalBox> MainBox;
+	TSharedPtr<SBorder> CustomSelectionBorder;
+	TSharedRef<SLauncherTableRow> RowContent = SNew(SLauncherTableRow, OwnerTable)
+	.ShowSelection(false)
+	[
+		SNew(SBorder)
+		.Padding(FMargin(4,4,4,0))
+		.BorderImage(FAppStyle::GetBrush("Brushes.Background"))
+		[
+			SAssignNew(CustomSelectionBorder, SBorder)
+			.BorderImage(FCustomLaunchUIStyle::GetBrush("WhiteGroupBorder"))
+			[
+				SAssignNew(MainBox, SHorizontalBox)
+
+				// platform icon
+				+SHorizontalBox::Slot()
+				.AutoWidth()
+				.VAlign(VAlign_Center)
+				.Padding(4)
+				[
+					SNew(SImage)
+					.DesiredSizeOverride(FVector2D(48,48))
+					.Image(this, &SCustomLaunchCustomProfileSelector::GetProfileImage, Profile)
+				]
+
+				// name and description etc.
+				+SHorizontalBox::Slot()
+				.FillWidth(1)
+				.Padding(4)
+				[
+					SAssignNew(DetailsBox, SVerticalBox)
+				]
+			]
+		]
+	];
+
+	// set the color attribute after the row itself has been created - this allows us to capture the row widget itself
+	CustomSelectionBorder->SetBorderBackgroundColor(TAttribute<FSlateColor>::CreateSP(this, &SCustomLaunchCustomProfileSelector::GetRowColor, RowContent) );
+	
+	if (bIsBasicLaunchProfile)
+	{
+		// add the name
+		DetailsBox->AddSlot()
+		.Padding(4)
+		.AutoHeight()
+		[
+			SNew(STextBlock)
+			.ColorAndOpacity(FColor::White)
+			.Text_Lambda( [Profile]() { return FText::FromString(Profile->GetName()); })
+		];
+
+		// add the quick device picker
+		DetailsBox->AddSlot()
+		.Padding(4)
+		.AutoHeight()
+		[
+			SNew(SHorizontalBox)
+
+			// device picker
+			+SHorizontalBox::Slot()
+			.AutoWidth()
+			.VAlign(VAlign_Center)
+			[
+				SNew(SCustomLaunchDeviceCombo)
+				.OnDeviceRemoved(this, &SCustomLaunchCustomProfileSelector::OnDeviceRemoved, Profile)
+				.OnSelectionChanged(this, &SCustomLaunchCustomProfileSelector::SetSelectedDevices, Profile)
+				.SelectedDevices(this, &SCustomLaunchCustomProfileSelector::GetSelectedDevices, Profile)
+				.AllPlatforms(true)
+			]
+
+			// device manager button
+			+SHorizontalBox::Slot()
+			.AutoWidth()
+			.VAlign(VAlign_Center)
+			.Padding(4,0)
+			[
+				SNew(SButton)
+				.OnClicked(this, &SCustomLaunchCustomProfileSelector::OnOpenDeviceManagerClicked)
+				.ButtonStyle(FAppStyle::Get(), "HoverHintOnly")
+				.ToolTipText(LOCTEXT("OpenDeviceManagerToolTip", "Open the Device Manager window, where you can setup and claim devices connected to your machine or shared on the network."))
+				[
+					SNew(SImage)
+					.Image(FAppStyle::Get().GetBrush("DeviceDetails.TabIcon"))
+					.ColorAndOpacity(FColor::White)
+				]
+			]
+		];
+
+		// fixme: current settings etc.?
+	}
+	else
+	{
+		TSharedPtr<SInlineEditableTextBlock> NameEditTextBox;
+		TSharedPtr<SInlineEditableTextBlock> DescriptionEditTextBox;
+
+		// name edit box
+		DetailsBox->AddSlot()
+		.Padding(4)
+		.AutoHeight()
+		[
+			SAssignNew(NameEditTextBox, SInlineEditableTextBlock)
+			.Text_Lambda( [Profile]() { return FText::FromString(Profile->GetName()); })
+			.ColorAndOpacity(FColor::White)
+			.OnTextCommitted( this, &SCustomLaunchCustomProfileSelector::SetProfileName, Profile )
+		];
+
+		// description
+		DetailsBox->AddSlot()
+		.Padding(4)
+		.AutoHeight()
+		[
+			SAssignNew(DescriptionEditTextBox, SInlineEditableTextBlock)
+			.Text_Lambda( [Profile]() { return FText::FromString(Profile->GetDescription()); })
+			.OnTextCommitted( this, &SCustomLaunchCustomProfileSelector::SetProfileDescription, Profile )
+			.OnExitEditingMode_Lambda( [this]() { Model->SortProfiles(); RefreshCustomProfileList(); } )
+			.ColorAndOpacity(FSlateColor::UseSubduedForeground())
+		];
+
+		NameEditTextBoxes.Add(Profile, NameEditTextBox);
+		DescriptionEditTextBoxes.Add(Profile, DescriptionEditTextBox);
+	}
+
+	// floating edit button
+	MainBox->AddSlot()
+	.VAlign(VAlign_Center)
+	.HAlign(HAlign_Right)
+	.Padding(8)
+	.AutoWidth()
+	[
+		SNew(SButton)
+		.OnClicked(this, &SCustomLaunchCustomProfileSelector::OnEditProfileClicked, Profile)
+		.ButtonStyle(FAppStyle::Get(), "HoverHintOnly")
+		.Visibility( this, &SCustomLaunchCustomProfileSelector::GetInlineEditButtonVisibility, RowContent->AsWidget(), Profile )
+		[
+			SNew(SImage)
+			.Image(FAppStyle::Get().GetBrush("Icons.Edit"))
+		]
+	];
+
+	// close editor button
+	MainBox->AddSlot()
+	.VAlign(VAlign_Center)
+	.HAlign(HAlign_Right)
+	.Padding(8)
+	.AutoWidth()
+	[
+		SNew(SButton)
+		.OnClicked(this, &SCustomLaunchCustomProfileSelector::OnCloseEditorClicked)
+		.ButtonStyle(FAppStyle::Get(), "HoverHintOnly")
+		.Visibility( this, &SCustomLaunchCustomProfileSelector::GetCloseEditorButtonVisibility, Profile)
+		[
+			SNew(SImage)
+			.Image(FCustomLaunchUIStyle::Get().GetBrush("SidePanelRightClose"))
+		]
+	];
+
+	return RowContent;
+}
+END_SLATE_FUNCTION_BUILD_OPTIMIZATION
+
+const FSlateBrush* SCustomLaunchCustomProfileSelector::GetProfileImage(ILauncherProfilePtr Profile) const
+{
+	if (Profile.IsValid() && Model->IsAdvancedProfile(Profile.ToSharedRef()))
+	{
+		return FAppStyle::GetBrush("Icons.Warning.Large");
+	}
+	else
+	{
+		const PlatformInfo::FTargetPlatformInfo* PlatformInfo = Model->GetPlatformInfo(Profile);
+		return FCustomLaunchUIStyle::GetProfileBrushForPlatform(PlatformInfo, EPlatformIconSize::Large);
+	}
+}
+
+FSlateColor SCustomLaunchCustomProfileSelector::GetRowColor(TSharedRef<SLauncherTableRow> TableRow) const
+{
+	if (TableRow->IsItemSelected())
+	{
+#if 1	// always maintain active state, as per UX design
+		return FStyleColors::Select;
+#else
+		if (CustomProfileListView->HasKeyboardFocus())
+		{
+			return FStyleColors::Select;
+		}
+		else
+		{
+			return FStyleColors::SelectInactive;
+		}
+#endif
+	}
+	else
+	{
+		if (TableRow->IsHovered())
+		{
+			return FStyleColors::Hover;
+		}
+		else
+		{
+			return FStyleColors::Panel;
+		}
+	}
+}
+
+void SCustomLaunchCustomProfileSelector::OnSelectionChanged( ILauncherProfilePtr Profile, ESelectInfo::Type SelectionMode )
+{
+	if (Profile.IsValid() && ensure(CustomProfileListView->IsItemSelected(Profile)))
+	{
+		Model->SelectProfile(Profile);
+	}
+}
+
+void SCustomLaunchCustomProfileSelector::OnProfileSelected( const ILauncherProfilePtr& NewProfile, const ILauncherProfilePtr& OldProfile)
+{
+	if (OldProfile.IsValid() && CustomProfileListView->IsItemSelected(OldProfile))
+	{
+		CustomProfileListView->SetItemSelection(OldProfile, false, ESelectInfo::Direct);
+	}
+
+	if (NewProfile.IsValid())
+	{
+		if (!CustomProfileListView->IsItemSelected(NewProfile))
+		{
+			CustomProfileListView->SetItemSelection(NewProfile, true, ESelectInfo::Direct);
+			CustomProfileListView->RequestScrollIntoView(NewProfile);
+		}
+	}
+}
+
+
+void SCustomLaunchCustomProfileSelector::OnCustomProfileAdded(const ILauncherProfileRef& AddedProfile)
+{
+	RefreshCustomProfileList();
+}
+
+
+
+void SCustomLaunchCustomProfileSelector::OnCustomProfileRemoved(const ILauncherProfileRef& RemovedProfile)
+{
+	NameEditTextBoxes.Remove(RemovedProfile.ToSharedPtr());
+	DescriptionEditTextBoxes.Remove(RemovedProfile.ToSharedPtr());
+
+	RefreshCustomProfileList();
+}
+
+
+
+void SCustomLaunchCustomProfileSelector::RefreshCustomProfileList()
+{
+	CustomProfileListView->RequestListRefresh();
+}
+
+
+TSharedPtr<SWidget> SCustomLaunchCustomProfileSelector::MakeContextMenu()
+{
+	ILauncherProfilePtr Profile = Model->GetSelectedProfile();
+	bool bIsBasicLaunchProfile = Model->IsBasicLaunchProfile(Profile);
+
+	const bool bShouldCloseWindowAfterMenuSelection = true;
+	FMenuBuilder MenuBuilder(bShouldCloseWindowAfterMenuSelection, nullptr);
+
+	MenuBuilder.AddMenuEntry(
+		LOCTEXT("CustomProfileAddLabel", "Add Profile"),
+		LOCTEXT("CustomProfileAddToolTip", "Add a new profile"),
+		FSlateIcon(FAppStyle::Get().GetStyleSetName(), "Icons.Plus"),
+		FUIAction(
+			FExecuteAction::CreateLambda( [this]() { OnProfileAdd.Execute(); } )
+		),
+		NAME_None);
+
+	MenuBuilder.AddMenuEntry(
+		LOCTEXT("CustomProfileEditLabel", "Edit Profile"),
+		LOCTEXT("CustomProfileEditToolTip", "Edit this profile"),
+		FSlateIcon(FAppStyle::Get().GetStyleSetName(), "Icons.Edit"),
+		FUIAction(
+			FExecuteAction::CreateSP(this, &SCustomLaunchCustomProfileSelector::OnCustomProfileEditClicked, Profile),
+			FCanExecuteAction(),
+			FGetActionCheckState::CreateSP(this, &SCustomLaunchCustomProfileSelector::GetCustomProfileEditCheckState, Profile)
+		),
+		NAME_None,
+		EUserInterfaceActionType::Check);
+
+	if (!bIsBasicLaunchProfile)
+	{
+		MenuBuilder.AddMenuSeparator();
+
+		MenuBuilder.AddMenuEntry(
+			LOCTEXT("CustomProfileRenameLabel", "Rename Profile"),
+			LOCTEXT("CustomProfileRenameToolTip", "Rename this profile"),
+			FSlateIcon(FAppStyle::Get().GetStyleSetName(), "Icons.Edit"),
+			FUIAction(
+				FExecuteAction::CreateSP(this, &SCustomLaunchCustomProfileSelector::OnCustomProfileRenameClicked, Profile)
+			),
+			NAME_None);
+
+		MenuBuilder.AddMenuEntry(
+			LOCTEXT("CustomProfileDuplicateLabel", "Duplicate Profile"),
+			LOCTEXT("CustomProfileDuplicateToolTip", "Duplicate this profile"),
+			FSlateIcon(FAppStyle::Get().GetStyleSetName(), "Icons.Duplicate"),
+			FUIAction(
+				FExecuteAction::CreateSP(this, &SCustomLaunchCustomProfileSelector::OnCustomProfileDuplicateClicked, Profile)
+			),
+			NAME_None);
+
+		MenuBuilder.AddMenuSeparator();
+
+		MenuBuilder.AddMenuEntry(
+			LOCTEXT("CustomProfileDeleteLabel", "Delete Profile"),
+			LOCTEXT("CustomProfileDeleteToolTip", "Deletes this profile"),
+			FSlateIcon(FAppStyle::Get().GetStyleSetName(), "Icons.Delete"),
+			FUIAction(
+				FExecuteAction::CreateSP(this, &SCustomLaunchCustomProfileSelector::OnCustomProfileDeleteClicked, Profile)
+			),
+			NAME_None);
+		
+	}
+
+	return MenuBuilder.MakeWidget();
+}
+
+void SCustomLaunchCustomProfileSelector::OnCustomProfileEditClicked(ILauncherProfilePtr Profile)
+{
+	if ( EditPanelVisible.Get() && Model->GetSelectedProfile() == Profile)
+	{
+		ChangeProfileEditorVisibility.ExecuteIfBound(false);
+	}
+	else
+	{
+		EditProfile(Profile);
+	}
+}
+
+void SCustomLaunchCustomProfileSelector::OnCustomProfileDuplicateClicked(ILauncherProfilePtr Profile)
+{
+	OnProfileDuplicate.ExecuteIfBound(Profile);
+}
+
+void SCustomLaunchCustomProfileSelector::OnCustomProfileDeleteClicked(ILauncherProfilePtr Profile)
+{
+	OnProfileDelete.ExecuteIfBound(Profile);
+}
+
+void SCustomLaunchCustomProfileSelector::OnCustomProfileRenameClicked(ILauncherProfilePtr Profile)
+{
+	OnProfileRename.ExecuteIfBound(Profile);
+}
+
+void SCustomLaunchCustomProfileSelector::OnDeviceRemoved( const FString DeviceID, ILauncherProfilePtr Profile)
+{
+	if (ensure(Profile->GetDeployedDeviceGroup() != nullptr) && Profile->GetDeployedDeviceGroup()->GetDeviceIDs().Contains(DeviceID))
+	{
+		Profile->GetDeployedDeviceGroup()->RemoveDevice(DeviceID);
+
+		OnProfileModified.ExecuteIfBound(Profile);
+	}
+}
+
+void SCustomLaunchCustomProfileSelector::SetSelectedDevices( const TArray<FString> DeviceIDs, ILauncherProfilePtr Profile)
+{
+	if (ensure(Profile->GetDeployedDeviceGroup() != nullptr))
+	{
+		Profile->GetDeployedDeviceGroup()->RemoveAllDevices();
+		for (const FString& DeviceID : DeviceIDs)
+		{
+			Profile->GetDeployedDeviceGroup()->AddDevice(DeviceID);
+		}
+		Model->UpdatedCookedPlatformsFromDeployDeviceProxy(Profile.ToSharedRef());
+
+		OnProfileModified.ExecuteIfBound(Profile);
+	}
+}
+
+TArray<FString> SCustomLaunchCustomProfileSelector::GetSelectedDevices(const ILauncherProfilePtr Profile) const
+{
+	if (ensure(Profile->GetDeployedDeviceGroup() != nullptr))
+	{
+		return Profile->GetDeployedDeviceGroup()->GetDeviceIDs();
+	}
+
+	return TArray<FString>();
+}
+
+FReply SCustomLaunchCustomProfileSelector::OnOpenDeviceManagerClicked()
+{
+	FGlobalTabmanager::Get()->TryInvokeTab(FTabId("DeviceManager"));
+	return FReply::Handled();
+}
+
+FReply SCustomLaunchCustomProfileSelector::OnEditProfileClicked(ILauncherProfilePtr Profile)
+{
+	EditProfile(Profile);
+
+	return FReply::Handled();
+}
+
+FReply SCustomLaunchCustomProfileSelector::OnCloseEditorClicked()
+{
+	ChangeProfileEditorVisibility.ExecuteIfBound(false);
+
+	return FReply::Handled();
+}
+
+EVisibility SCustomLaunchCustomProfileSelector::GetCloseEditorButtonVisibility(ILauncherProfilePtr Profile) const
+{
+	if ( EditPanelVisible.Get() && Model->GetSelectedProfile() == Profile)
+	{
+		return EVisibility::Visible;
+	}
+	else
+	{
+		return EVisibility::Collapsed;
+	}
+}
+
+ECheckBoxState SCustomLaunchCustomProfileSelector::GetCustomProfileEditCheckState(ILauncherProfilePtr Profile) const
+{
+	if ( EditPanelVisible.Get() && Model->GetSelectedProfile() == Profile)
+	{
+		return ECheckBoxState::Checked;
+	}
+	else
+	{
+		return ECheckBoxState::Unchecked;
+	}
+}
+
+
+void SCustomLaunchCustomProfileSelector::EditProfile(ILauncherProfilePtr Profile)
+{
+	if ( EditPanelVisible.Get() && Model->GetSelectedProfile() == Profile)
+	{
+		ChangeProfileEditorVisibility.ExecuteIfBound(false);
+	}
+	else
+	{
+		OnProfileEdit.ExecuteIfBound(Profile);
+	}
+}
+
+EVisibility SCustomLaunchCustomProfileSelector::GetInlineEditButtonVisibility(TSharedRef<SWidget> RowWidget, ILauncherProfilePtr Profile) const
+{
+	if (EditPanelVisible.Get())
+	{
+		return EVisibility::Collapsed;
+	}
+	else if (RowWidget->IsHovered())
+	{
+		return EVisibility::Visible;
+	}
+	else
+	{
+		return EVisibility::Hidden;
+	}
+}
+
+
+void SCustomLaunchCustomProfileSelector::StartEditProfileName(ILauncherProfilePtr Profile)
+{
+	if (Profile.IsValid() && NameEditTextBoxes.Contains(Profile))
+	{
+		NameEditTextBoxes.FindRef(Profile)->EnterEditingMode();
+	}
+}
+
+void SCustomLaunchCustomProfileSelector::StartEditProfileDescription(ILauncherProfilePtr Profile)
+{
+	if (Profile.IsValid() && DescriptionEditTextBoxes.Contains(Profile))
+	{
+		DescriptionEditTextBoxes.FindRef(Profile)->EnterEditingMode();
+	}
+}
+
+
+void SCustomLaunchCustomProfileSelector::SetProfileName(const FText& NewText, ETextCommit::Type InTextCommit, ILauncherProfilePtr Profile)
+{
+	if (Profile.IsValid())
+	{
+		Model->GetProfileManager()->ChangeProfileName(Profile.ToSharedRef(), NewText.ToString());
+
+		if (InTextCommit == ETextCommit::OnEnter || InTextCommit == ETextCommit::OnUserMovedFocus)
+		{
+			StartEditProfileDescription(Profile);
+		}
+		else
+		{
+			Model->SortProfiles();
+			RefreshCustomProfileList();
+		}
+	}
+}
+
+void SCustomLaunchCustomProfileSelector::SetProfileDescription(const FText& NewText, ETextCommit::Type InTextCommit, ILauncherProfilePtr Profile)
+{
+	if (Profile.IsValid())
+	{
+		Profile->SetDescription(NewText.ToString());
+		Model->GetProfileManager()->SaveJSONProfile(Profile.ToSharedRef());	
+	}
+}
+
+
+
+
+
+#undef LOCTEXT_NAMESPACE
diff --git a/Engine/Plugins/Developer/CustomLaunchUI/Source/Private/Widgets/Profiles/SCustomLaunchCustomProfileSelector.h b/Engine/Plugins/Developer/CustomLaunchUI/Source/Private/Widgets/Profiles/SCustomLaunchCustomProfileSelector.h
new file mode 100644
--- /dev/null
+++ b/Engine/Plugins/Developer/CustomLaunchUI/Source/Private/Widgets/Profiles/SCustomLaunchCustomProfileSelector.h
@@ -0,0 +1,92 @@
+// Copyright Epic Games, Inc. All Rights Reserved.
+
+#pragma once
+
+#include "Internationalization/Text.h"
+#include "Framework/SlateDelegates.h"
+#include "Widgets/SCompoundWidget.h"
+#include "Widgets/Views/STableRow.h"
+#include "Model/CustomLaunchUIModel.h"
+
+class SWidget;
+class ITableRow;
+class STableViewBase;
+class SInlineEditableTextBlock;
+template<typename T> class SListView;
+
+
+class SCustomLaunchCustomProfileSelector 
+	: public SCompoundWidget
+{
+public:
+	DECLARE_DELEGATE_OneParam(FSetProfileEditorVisible, bool)
+
+	SLATE_BEGIN_ARGS(SCustomLaunchCustomProfileSelector) {}
+		SLATE_EVENT(FOnClicked, OnProfileAdd)
+		SLATE_EVENT(CustomLaunchUI::FOnProfileClicked, OnProfileEdit)
+		SLATE_EVENT(CustomLaunchUI::FOnProfileClicked, OnProfileDuplicate)
+		SLATE_EVENT(CustomLaunchUI::FOnProfileClicked, OnProfileDelete)
+		SLATE_EVENT(CustomLaunchUI::FOnProfileClicked, OnProfileRename)
+		SLATE_EVENT(CustomLaunchUI::FOnProfileClicked, OnProfileModified)
+		SLATE_EVENT(FSetProfileEditorVisible, ChangeProfileEditorVisibility)
+		SLATE_ATTRIBUTE(bool, EditPanelVisible)
+	SLATE_END_ARGS()
+
+public:
+	void Construct(const FArguments& InArgs, const TSharedRef<CustomLaunchUI::FModel>& InModel);
+	~SCustomLaunchCustomProfileSelector();
+
+	void StartEditProfileName(ILauncherProfilePtr Profile);
+	void StartEditProfileDescription(ILauncherProfilePtr Profile);
+
+private:
+	typedef STableRow<TSharedPtr<ILauncherProfilePtr>> SLauncherTableRow;
+
+	FOnClicked OnProfileAdd;
+	CustomLaunchUI::FOnProfileClicked OnProfileEdit;
+	CustomLaunchUI::FOnProfileClicked OnProfileDelete;
+	CustomLaunchUI::FOnProfileClicked OnProfileDuplicate;
+	CustomLaunchUI::FOnProfileClicked OnProfileRename;
+	CustomLaunchUI::FOnProfileClicked OnProfileModified;
+	FSetProfileEditorVisible ChangeProfileEditorVisibility;
+	TAttribute<bool> EditPanelVisible;
+
+	TSharedRef<ITableRow> GenerateCustomProfileRow(ILauncherProfilePtr Profile, const TSharedRef<STableViewBase>& OwnerTable);
+	const FSlateBrush* GetProfileImage(ILauncherProfilePtr Profile) const;
+	FSlateColor GetRowColor(TSharedRef<SLauncherTableRow> TableRow) const;
+
+	void OnSelectionChanged( ILauncherProfilePtr Profile, ESelectInfo::Type SelectionMode );
+	void OnProfileSelected( const ILauncherProfilePtr& NewProfile, const ILauncherProfilePtr& OldProfile);
+	void OnCustomProfileAdded(const ILauncherProfileRef& AddedProfile);
+	void OnCustomProfileRemoved(const ILauncherProfileRef& RemovedProfile);
+	void OnCustomProfileEditClicked(ILauncherProfilePtr Profile);
+	void OnCustomProfileDuplicateClicked(ILauncherProfilePtr Profile);
+	void OnCustomProfileDeleteClicked(ILauncherProfilePtr Profile);
+	void OnCustomProfileRenameClicked(ILauncherProfilePtr Profile);
+	FReply OnOpenDeviceManagerClicked();
+	FReply OnEditProfileClicked(ILauncherProfilePtr Profile);
+	FReply OnCloseEditorClicked();
+	EVisibility GetInlineEditButtonVisibility(TSharedRef<SWidget> RowWidget, ILauncherProfilePtr Profile) const;
+	EVisibility GetCloseEditorButtonVisibility(ILauncherProfilePtr Profile) const;
+	ECheckBoxState GetCustomProfileEditCheckState(ILauncherProfilePtr Profile) const;
+
+	void RefreshCustomProfileList();
+
+	void OnDeviceRemoved(const FString DeviceID, ILauncherProfilePtr Profile);
+	void SetSelectedDevices(const TArray<FString> DeviceIDs, ILauncherProfilePtr Profile );
+	TArray<FString> GetSelectedDevices(const ILauncherProfilePtr Profile) const;
+
+	TSharedPtr<SWidget> MakeContextMenu();
+	void EditProfile(ILauncherProfilePtr Profile);
+
+	void SetProfileName(const FText& NewText, ETextCommit::Type InTextCommit, ILauncherProfilePtr Profile);
+	void SetProfileDescription(const FText& NewText, ETextCommit::Type InTextCommit, ILauncherProfilePtr Profile);
+
+	TSharedPtr<SListView<ILauncherProfilePtr>> CustomProfileListView;
+
+	TMap<ILauncherProfilePtr, TSharedPtr<SInlineEditableTextBlock>> NameEditTextBoxes;
+	TMap<ILauncherProfilePtr, TSharedPtr<SInlineEditableTextBlock>> DescriptionEditTextBoxes;
+
+private:
+	TSharedPtr<CustomLaunchUI::FModel> Model;
+};
diff --git a/Engine/Plugins/Developer/CustomLaunchUI/Source/Private/Widgets/Profiles/SCustomLaunchProfilesPanel.cpp b/Engine/Plugins/Developer/CustomLaunchUI/Source/Private/Widgets/Profiles/SCustomLaunchProfilesPanel.cpp
new file mode 100644
--- /dev/null
+++ b/Engine/Plugins/Developer/CustomLaunchUI/Source/Private/Widgets/Profiles/SCustomLaunchProfilesPanel.cpp
@@ -0,0 +1,842 @@
+// Copyright Epic Games, Inc. All Rights Reserved.
+
+#include "Widgets/Profiles/SCustomLaunchProfilesPanel.h"
+#include "SlateOptMacros.h"
+#include "Styling/StyleColors.h"
+#include "Styling/CustomLaunchUIStyle.h"
+#include "Widgets/SBoxPanel.h"
+#include "Widgets/Layout/SSplitter.h"
+#include "Widgets/Layout/SSpacer.h"
+#include "Widgets/Layout/SExpandableArea.h"
+#include "Widgets/Layout/SScrollBorder.h"
+#include "Widgets/Text/STextBlock.h"
+#include "Widgets/Images/SImage.h"
+#include "Widgets/Input/SButton.h"
+#include "Widgets/Input/SComboButton.h"
+#include "Widgets/Input/SCheckBox.h"
+#include "Widgets/Shared/SCustomLaunchProjectCombo.h"
+#include "Widgets/Shared/SCustomLaunchBuildTargetCombo.h"
+#include "Widgets/Profiles/SCustomLaunchCustomProfileSelector.h"
+#include "Widgets/Output/SCustomLaunchLaunchPanel.h"
+#include "Widgets/Output/SCustomLaunchOutputLog.h"
+#include "Classes/LauncherProfileSettings.h"
+#include "ILauncher.h"
+#include "ITargetDeviceProxyManager.h"
+#include "ITargetDeviceProxy.h"
+#include "IDetailsView.h"
+#include "Modules/ModuleManager.h"
+#include "PropertyEditorModule.h"
+#include "Framework/Docking/TabManager.h"
+#include "PlatformInfo.h"
+
+
+#define LOCTEXT_NAMESPACE "SCustomLaunchProfilesPanel"
+
+
+SCustomLaunchProfilesPanel::SCustomLaunchProfilesPanel()
+{
+}
+
+
+SCustomLaunchProfilesPanel::~SCustomLaunchProfilesPanel()
+{
+	if (CurrentProfileSettingsObject != nullptr)
+	{
+		CurrentProfileSettingsObject->RemoveFromRoot();
+	}
+
+	if (Model.IsValid())
+	{
+		Model->OnProfileSelected().RemoveAll(this);
+	}
+}
+
+
+BEGIN_SLATE_FUNCTION_BUILD_OPTIMIZATION
+void SCustomLaunchProfilesPanel::Construct(const FArguments& InArgs, const TSharedRef<CustomLaunchUI::FModel>& InModel)
+{
+	Model = InModel;
+	OnProfileLaunchClicked = InArgs._OnProfileLaunchClicked;
+
+	TSharedPtr<CustomLaunchUI::FLaunchLogTextLayoutMarshaller> LaunchLogTextMarshaller = MakeShared<CustomLaunchUI::FLaunchLogTextLayoutMarshaller>(InModel);
+	OutputLog = SNew(SCustomLaunchOutputLog, InModel, LaunchLogTextMarshaller.ToSharedRef());
+	ChildSlot
+	[
+		SAssignNew(LogAreaSplitter, SSplitter)
+		.Orientation(Orient_Vertical)
+		.PhysicalSplitterHandleSize(8)
+		.Style(FAppStyle::Get(), "SplitterPanel")
+
+		// profile editor
+		+SSplitter::Slot()
+		.Resizable(false)
+		[
+			CreateProfilesPanel()
+		]
+
+		// bottom-docked log, for referencing the most recently completed build
+		+SSplitter::Slot()
+		.SizeRule(SSplitter::ESizeRule::SizeToContent)
+		.Resizable(false)
+		[
+			SAssignNew(LogExpandableArea, SExpandableArea)
+			.InitiallyCollapsed(true)
+			.BorderImage(FAppStyle::GetBrush("Brushes.Header"))
+			.OnAreaExpansionChanged(this, &SCustomLaunchProfilesPanel::OnLogAreaExpansionChanged)
+			.Padding(FMargin(0,8,0,0))
+			.HeaderContent()
+			[
+				SNew(SHorizontalBox)
+
+				// output log title
+				+SHorizontalBox::Slot()
+				.FillWidth(1)
+				.HAlign(HAlign_Left)
+				.VAlign(VAlign_Center)
+				.Padding(2)
+				[
+					SNew(STextBlock)
+					.Text(LOCTEXT("OutputLogAreaTitle", "Output Log"))
+				]
+
+				// output log filter box
+				+SHorizontalBox::Slot()
+				.HAlign(HAlign_Right)
+				.VAlign(VAlign_Center)
+				[
+					OutputLog->CreateFilterWidget()
+				]
+			]
+			.BodyContent()
+			[
+				OutputLog.ToSharedRef()
+			]
+		]
+	];
+
+	// hide profile editor by default
+	SetProfileEditorVisible(false);
+
+
+	Model->OnProfileSelected().AddSP(this, &SCustomLaunchProfilesPanel::OnProfileSelected);
+
+	// wait one tick then select the Basic Launch profile. this should give enough time for the host device proxy to be created
+	ExecuteOnGameThread(UE_SOURCE_LOCATION, [this]()
+	{
+		Model->SelectProfile(Model->GetBasicLaunchProfile());
+	});
+}
+END_SLATE_FUNCTION_BUILD_OPTIMIZATION
+
+
+BEGIN_SLATE_FUNCTION_BUILD_OPTIMIZATION
+TSharedRef<SWidget> SCustomLaunchProfilesPanel::CreateProfilesPanel()
+{
+	TSharedPtr<SVerticalBox> ProfilesPanel = SNew(SVerticalBox);
+
+	// global project & target configuration (for UnrealFrontend only - not used in editor)
+	if (!GIsEditor)
+	{
+		ProfilesPanel->AddSlot()
+		.AutoHeight()
+		.Padding(8,8,8,0)
+		[
+			SNew(SHorizontalBox)
+
+			// Default Project selector label
+			+SHorizontalBox::Slot()
+			.AutoWidth()
+			.VAlign(VAlign_Center)
+			.Padding(4,0)
+			[
+				SNew(STextBlock)
+				.Text(LOCTEXT("DefaultProjectSelectLabel", "Default Project"))
+			]
+
+			// Default Project selector
+			+SHorizontalBox::Slot()
+			.AutoWidth()
+			.VAlign(VAlign_Center)
+			.Padding(2,0)
+			[
+				SNew(SCustomLaunchProjectCombo)
+				.OnSelectionChanged(this, &SCustomLaunchProfilesPanel::SetDefaultProjectPath)
+				.SelectedProject(this, &SCustomLaunchProfilesPanel::GetDefaultProjectPath)
+				.ShowAnyProjectOption(false)
+			]
+
+			// Default Build Target selector label
+			+SHorizontalBox::Slot()
+			.AutoWidth()
+			.VAlign(VAlign_Center)
+			.Padding(16,0,2,0)
+			[
+				SNew(STextBlock)
+				.Text(LOCTEXT("DefaultBuildTargetSelectLabel", "Target"))
+			]
+
+			// Default Build Target selector
+			+SHorizontalBox::Slot()
+			.AutoWidth()
+			.VAlign(VAlign_Center)
+			.Padding(2,0)
+			[
+				SNew(SCustomLaunchBuildTargetCombo)
+				.OnSelectionChanged(this, &SCustomLaunchProfilesPanel::SetDefaultBuildTarget)
+				.SelectedBuildTarget(this, &SCustomLaunchProfilesPanel::GetDefaultBuildTarget)
+				.SelectedProject(this, &SCustomLaunchProfilesPanel::GetDefaultProjectPath)
+			]
+		];		
+	}
+
+	// main profile banner
+	ProfilesPanel->AddSlot()
+	.AutoHeight()
+	.Padding(8,8,8,0)
+	[
+		CreateProfileEditorToolbarWidget()
+	];
+
+	// main profile selection & editing panel
+	ProfilesPanel->AddSlot()
+	.FillHeight(1)
+	.Padding(8,8,8,0)
+	[
+		SAssignNew(ProfileEditorSplitter, SSplitter)
+		.Orientation(Orient_Horizontal)
+		.PhysicalSplitterHandleSize(8)
+		.Style(FAppStyle::Get(), "SplitterPanel")
+			
+		// current profile selection panel
+		+SSplitter::Slot()
+		[
+			CreateProfileSelectorWidget()
+		]
+		
+		// Profile editor panel
+		+SSplitter::Slot()
+		[
+			CreateProfileEditorWidget()
+		]
+	];
+
+	return ProfilesPanel.ToSharedRef();
+
+}
+END_SLATE_FUNCTION_BUILD_OPTIMIZATION
+
+
+
+BEGIN_SLATE_FUNCTION_BUILD_OPTIMIZATION
+TSharedRef<SWidget> SCustomLaunchProfilesPanel::CreateProfileSelectorWidget()
+{
+	return SNew(SVerticalBox)
+
+	// top banner
+	+SVerticalBox::Slot()
+	.AutoHeight()
+	[
+		SNew(SBorder)
+		.BorderImage(FAppStyle::GetBrush("Brushes.Header"))
+		[
+			SNew(SHorizontalBox)
+
+			// bannel label
+			+SHorizontalBox::Slot()
+			.FillWidth(1)
+			.VAlign(VAlign_Center)
+			.Padding(4,2)
+			[
+				SNew(STextBlock)
+				.Text(LOCTEXT("ProfilesLabel", "Launch Profiles"))
+			]
+
+			// new profile button
+			+SHorizontalBox::Slot()
+			.AutoWidth()
+			.VAlign(VAlign_Center)
+			.HAlign(HAlign_Right)
+			.Padding(4,2)
+			[
+				SNew(SButton)
+				.OnClicked(this, &SCustomLaunchProfilesPanel::OnCreateNewCustomProfileClicked)
+				.ToolTipText(LOCTEXT("NewProfileToolTip", "Create a new custom profile"))
+				[
+					SNew(SHorizontalBox)
+				
+					+SHorizontalBox::Slot()
+					.AutoWidth()
+					.VAlign(VAlign_Center)
+					.Padding(1)
+					[
+						SNew(SImage)
+						.Image(FAppStyle::Get().GetBrush("Icons.Plus"))
+						.ColorAndOpacity(FStyleColors::AccentGreen)
+					]
+
+					+SHorizontalBox::Slot()
+					.AutoWidth()
+					.VAlign(VAlign_Center)
+					.Padding(1)
+					[
+						SNew(STextBlock)
+						.Text(LOCTEXT("CustomizeProfileLabel", "Create Launch Profile"))
+					]
+				]
+			]
+		]
+	]
+
+	// profiles list
+	+SVerticalBox::Slot()
+	.FillHeight(1)
+	[
+		SAssignNew(ProfileSelector, SCustomLaunchCustomProfileSelector, Model.ToSharedRef())
+		.OnProfileAdd(this, &SCustomLaunchProfilesPanel::OnCreateNewCustomProfileClicked)
+		.OnProfileDelete(this, &SCustomLaunchProfilesPanel::OnProfileDelete)
+		.OnProfileDuplicate(this, &SCustomLaunchProfilesPanel::OnProfileDuplicate)
+		.OnProfileEdit(this, &SCustomLaunchProfilesPanel::OnProfileEdit)
+		.OnProfileRename(this, &SCustomLaunchProfilesPanel::OnProfileRename)
+		.OnProfileModified(this, &SCustomLaunchProfilesPanel::OnProfileModified)
+		.ChangeProfileEditorVisibility(this, &SCustomLaunchProfilesPanel::SetProfileEditorVisible)
+		.EditPanelVisible(this, &SCustomLaunchProfilesPanel::IsProfileEditorVisible)
+	]
+	;
+}
+END_SLATE_FUNCTION_BUILD_OPTIMIZATION
+
+
+BEGIN_SLATE_FUNCTION_BUILD_OPTIMIZATION
+TSharedRef<SWidget> SCustomLaunchProfilesPanel::CreateProfileEditorWidget()
+{
+	FDetailsViewArgs DetailsViewArgs;
+	DetailsViewArgs.NameAreaSettings = FDetailsViewArgs::HideNameArea;
+	DetailsViewArgs.bAllowSearch = false;
+	DetailsViewArgs.bHideSelectionTip = true;
+	DetailsViewArgs.bShowPropertyMatrixButton = false;
+	DetailsViewArgs.bShowModifiedPropertiesOption = false;
+	DetailsViewArgs.bShowHiddenPropertiesWhilePlayingOption = false;
+	DetailsViewArgs.bAllowMultipleTopLevelObjects = false;
+	DetailsViewArgs.bShowKeyablePropertiesOption = false;
+	DetailsViewArgs.bShowAnimatedPropertiesOption = false;
+
+	FPropertyEditorModule& PropertyEditorModule = FModuleManager::GetModuleChecked<FPropertyEditorModule>("PropertyEditor");
+	DetailsView = PropertyEditorModule.CreateDetailView(DetailsViewArgs);
+
+	return SNew(SVerticalBox)
+	.Visibility_Lambda( [this] { return Model->GetSelectedProfile().IsValid() && IsProfileEditorVisible() ? EVisibility::Visible :EVisibility::Collapsed; } )
+
+	// advanced profile warning banner
+	+SVerticalBox::Slot()
+	.AutoHeight()
+	[
+		CreateAdvancedProfileWarningWidget()
+	]
+
+	// profile property editor
+	+SVerticalBox::Slot()
+	.FillHeight(1)
+	[
+		SNew(SScrollBorder, DetailsView.ToSharedRef())
+		[
+			DetailsView.ToSharedRef()
+		]
+	]
+
+	// hide properties buttom button
+	+SVerticalBox::Slot()
+	.AutoHeight()
+	[
+		SNew(SBorder)
+		.BorderImage(FAppStyle::GetBrush("Brushes.Header"))
+		.Visibility_Lambda( [this]() { return IsProfileEditorVisible() ? EVisibility::Visible : EVisibility::Collapsed; } )
+		[
+			SNew(SHorizontalBox)
+
+			+SHorizontalBox::Slot()
+			.FillWidth(1)
+			[
+				SNew(SSpacer)
+			]
+
+			+SHorizontalBox::Slot()
+			.AutoWidth()
+			.VAlign(VAlign_Center)
+			.HAlign(HAlign_Right)
+			.Padding(4,4)
+			[
+				SNew(SButton)
+				.OnClicked_Lambda( [this]() { SetProfileEditorVisible(false); return FReply::Handled();} )
+				.ButtonStyle(FAppStyle::Get(), "Button")
+				.Text(LOCTEXT("HidePropertiesLabel", "Hide Properties"))
+			]
+		]
+	]
+	;
+
+}
+END_SLATE_FUNCTION_BUILD_OPTIMIZATION
+
+BEGIN_SLATE_FUNCTION_BUILD_OPTIMIZATION
+TSharedRef<SWidget> SCustomLaunchProfilesPanel::CreateProfileEditorToolbarWidget()
+{
+	return SNew(SBorder)
+	.BorderImage(FAppStyle::GetBrush("Brushes.Header"))
+	.Padding(16)
+	[
+		SNew(SVerticalBox)
+				
+		// profile details & control buttons
+		+SVerticalBox::Slot()
+		.Padding(0)
+		[
+			SNew(SHorizontalBox)
+
+			// profile icon
+			+SHorizontalBox::Slot()
+			.AutoWidth()
+			.Padding(4,0)
+			.VAlign(VAlign_Center)
+			[
+				SNew(SImage)
+				.DesiredSizeOverride(FVector2D(44,44))
+				.Image( this, &SCustomLaunchProfilesPanel::GetSelectedProfileImage )
+			]
+
+			// profile details
+			+SHorizontalBox::Slot()
+			.AutoWidth()
+			.Padding(4)
+			[
+				SNew(SVerticalBox)
+
+				// profile name
+				+SVerticalBox::Slot()
+				.AutoHeight()
+				.Padding(2)
+				[
+					SNew(STextBlock)
+					.Text( this, &SCustomLaunchProfilesPanel::GetSelectedProfileName )
+				]
+
+				// profile description
+				+SVerticalBox::Slot()
+				.AutoHeight()
+				.Padding(2)
+				[
+					SNew(STextBlock)
+					.Text( this, &SCustomLaunchProfilesPanel::GetSelectedProfileDescription )
+					.ColorAndOpacity(FSlateColor::UseSubduedForeground())
+				]
+			]
+
+			// spacer
+			+SHorizontalBox::Slot()
+			.FillWidth(1)
+			[
+				SNew(SSpacer)
+			]
+
+			// prototype warning
+			+SHorizontalBox::Slot()
+			.AutoWidth()
+			.Padding(8)
+			.VAlign(VAlign_Top)
+			[
+				SNew(STextBlock)
+				.TextStyle(FAppStyle::Get(), "NormalText.Important")
+				.Text(LOCTEXT("PrototypeWarning", "This is a prototype and not ready for use!"))
+				.ColorAndOpacity(FAppStyle::Get().GetSlateColor("Colors.Warning"))
+			]
+
+			// launch button
+			+SHorizontalBox::Slot()
+			.AutoWidth()
+			.Padding(4,0,0,0)
+			.VAlign(VAlign_Top)
+			[
+				SNew(SButton)
+				.OnClicked( this, &SCustomLaunchProfilesPanel::OnLaunchButtonClicked )
+				.IsEnabled( this, &SCustomLaunchProfilesPanel::IsLaunchButtonEnabled )
+				.ButtonStyle(FAppStyle::Get(), "PrimaryButton")
+				.ToolTipText( this, &SCustomLaunchProfilesPanel::GetLaunchButtonToolTipText )
+				.ContentPadding(4)
+				[
+					SNew(SHorizontalBox)
+				
+					+SHorizontalBox::Slot()
+					.AutoWidth()
+					.VAlign(VAlign_Center)
+					.Padding(1,0)
+					[
+						SNew(SImage)
+						.Image( this, &SCustomLaunchProfilesPanel::GetLaunchButtonImage )
+						.ColorAndOpacity(FSlateColor::UseForeground())
+					]
+
+					+SHorizontalBox::Slot()
+					.AutoWidth()
+					.VAlign(VAlign_Center)
+					.Padding(1,0)
+					[
+						SNew(STextBlock)
+						.Text(LOCTEXT("DeviceProxyLaunchButton", "Launch"))
+					]
+				]
+			]
+		]
+	];
+}
+END_SLATE_FUNCTION_BUILD_OPTIMIZATION
+
+BEGIN_SLATE_FUNCTION_BUILD_OPTIMIZATION
+TSharedRef<SWidget> SCustomLaunchProfilesPanel::CreateAdvancedProfileWarningWidget()
+{
+	return SNew(SBorder)
+	.BorderImage(FAppStyle::GetBrush("ToolPanel.LightGroupBorder"))
+	.BorderBackgroundColor(FAppStyle::Get().GetSlateColor("Colors.Warning"))
+	.Visibility_Lambda( [this]() { return Model->GetSelectedProfile().IsValid() && Model->IsAdvancedProfile(Model->GetSelectedProfile().ToSharedRef()) ? EVisibility::Visible : EVisibility::Collapsed; })
+	.Padding(4)
+	[
+		SNew(SHorizontalBox)
+
+		// Notice
+		+SHorizontalBox::Slot()
+		.FillWidth(1)
+		.Padding(16, 0)
+		.VAlign(VAlign_Center)
+		[
+			SNew(SVerticalBox)
+
+			+SVerticalBox::Slot()
+			.AutoHeight()
+			.Padding(2)
+			[
+				SNew(STextBlock)
+				.TextStyle(FAppStyle::Get(), "LargeText")
+				.Text(LOCTEXT("AdvancedProfileWarning", "Profile contains advanced/legacy settings"))
+			]
+
+			+SVerticalBox::Slot()
+			.AutoHeight()
+			.Padding(2)
+			[
+				SNew(STextBlock)
+				.TextStyle(FAppStyle::Get(), "NormalText.Important")
+				.Text(LOCTEXT("AdvancedProfileWarningDetail", "These settings may be multiple cook platforms or cultures, or distribution packages, DLC or patch builds. Please use Project Launcher to edit the advanced properties in this profile."))
+			]
+		]
+
+		// Button
+		+SHorizontalBox::Slot()
+		.AutoWidth()
+		.VAlign(VAlign_Center)
+		[
+			SNew(SButton)
+			.OnClicked_Lambda( []() { FGlobalTabmanager::Get()->TryInvokeTab(FTabId("ProjectLauncher")); return FReply::Handled(); })
+			.ButtonStyle(FAppStyle::Get(), "HoverHintOnly")
+			.ToolTipText(LOCTEXT("OpenDeviceProjectLauncherTip", "Open the legacy Project Launcher."))
+			[
+				SNew(SImage)
+				.Image(FAppStyle::Get().GetBrush("Launcher.TabIcon"))
+			]
+		]
+	]
+	;
+}
+END_SLATE_FUNCTION_BUILD_OPTIMIZATION
+
+
+
+void SCustomLaunchProfilesPanel::OnProfileEdit(const ILauncherProfilePtr& Profile)
+{
+	Model->SelectProfile(Profile);
+	SetProfileEditorVisible(true);
+}
+
+
+
+
+void SCustomLaunchProfilesPanel::OnProfileDelete(const ILauncherProfilePtr& Profile)
+{
+	if (Profile.IsValid() && !Model->IsBasicLaunchProfile(Profile))
+	{
+		FText MessageText = LOCTEXT("ProfileDeleteConfirm", "Are you sure you want to delete this profile?");
+		if (FPlatformMisc::MessageBoxExt(EAppMsgType::YesNo, *MessageText.ToString(), *Profile->GetName()) == EAppReturnType::Yes)
+		{
+			Model->GetProfileManager()->RemoveProfile(Profile.ToSharedRef());
+			Model->SelectProfile(Model->GetBasicLaunchProfile());
+		}
+	}
+}
+
+
+void SCustomLaunchProfilesPanel::OnProfileDuplicate(const ILauncherProfilePtr& Profile)
+{
+	if (Profile.IsValid() && !Model->IsBasicLaunchProfile(Profile))
+	{
+		ILauncherProfilePtr NewProfile = Model->CloneCustomProfile(Profile.ToSharedRef());
+
+		if (NewProfile.IsValid())
+		{
+			Model->GetProfileManager()->AddProfile(NewProfile.ToSharedRef());
+
+			Model->SelectProfile(NewProfile);
+			SetProfileEditorVisible(true);
+
+			ProfileSelector->StartEditProfileName(NewProfile);
+		}
+	}
+}
+
+void SCustomLaunchProfilesPanel::OnProfileRename(const ILauncherProfilePtr& Profile)
+{
+	if (Profile.IsValid() && !Model->IsBasicLaunchProfile(Profile))
+	{
+		Model->SelectProfile(Profile);
+
+		ProfileSelector->StartEditProfileName(Profile);
+	}
+}
+
+
+void SCustomLaunchProfilesPanel::OnProfileSave(const ILauncherProfilePtr& Profile)
+{
+	if (Profile.IsValid() && !Model->IsBasicLaunchProfile(Profile))
+	{
+		Model->GetProfileManager()->AddProfile(Profile.ToSharedRef());
+	}
+	else
+	{
+		checkNoEntry();
+	}
+}
+
+
+
+FReply SCustomLaunchProfilesPanel::OnLaunchButtonClicked()
+{
+	OnProfileLaunchClicked.ExecuteIfBound(Model->GetSelectedProfile());
+	return FReply::Handled();
+}
+
+bool SCustomLaunchProfilesPanel::IsLaunchButtonEnabled() const
+{
+	ILauncherProfilePtr Profile = Model->GetSelectedProfile();
+	return Profile.IsValid() && Profile->IsValidForLaunch();
+}
+
+FText SCustomLaunchProfilesPanel::GetLaunchButtonToolTipText() const
+{
+	ILauncherProfilePtr Profile = Model->GetSelectedProfile();
+	return CustomLaunchUI::GetProfileLaunchErrorMessage(Profile);
+}
+
+const FSlateBrush* SCustomLaunchProfilesPanel::GetLaunchButtonImage() const
+{
+	return IsLaunchButtonEnabled() ? FAppStyle::Get().GetBrush("Icons.Launch") : FAppStyle::Get().GetBrush("Icons.Error");
+}
+
+
+
+FReply SCustomLaunchProfilesPanel::OnCreateNewCustomProfileClicked()
+{
+	// create a new profile
+	ILauncherProfileRef NewProfile = Model->CreateCustomProfile(TEXT("New Profile"));
+	NewProfile->AssignId();
+	Model->GetProfileManager()->AddProfile(NewProfile);
+
+	Model->SelectProfile(NewProfile);
+
+	// wait one tick then begin editing. this gives time for the list view to create the widgets for the new item
+	ExecuteOnGameThread(UE_SOURCE_LOCATION, [this, NewProfile]()
+	{
+		ProfileSelector->StartEditProfileName(NewProfile.ToSharedPtr());
+	});
+
+	return FReply::Handled();
+}
+
+FReply SCustomLaunchProfilesPanel::OnOpenDeviceManagerClicked()
+{
+	FGlobalTabmanager::Get()->TryInvokeTab(FTabId("DeviceManager"));
+	return FReply::Handled();
+}
+
+void SCustomLaunchProfilesPanel::OnProfileSelected( const ILauncherProfilePtr& NewProfile, const ILauncherProfilePtr& OldProfile)
+{
+	DetailsView->SetObject(nullptr);
+	if (CurrentProfileSettingsObject != nullptr)
+	{
+		CurrentProfileSettingsObject->RemoveFromRoot();
+		CurrentProfileSettingsObject = nullptr;
+	}
+
+	if (NewProfile.IsValid() && !Model->IsAdvancedProfile(NewProfile.ToSharedRef()))
+	{
+		if (Model->IsBasicLaunchProfile(NewProfile))
+		{
+			CurrentProfileSettingsObject = NewObject<ULauncherBasicLaunchProfileSettings>();
+		}
+		else
+		{
+			CurrentProfileSettingsObject = NewObject<ULauncherCustomProfileSettings>();
+		}
+
+		CurrentProfileSettingsObject->AddToRoot();
+		CurrentProfileSettingsObject->SetProfile(Model, NewProfile);
+		DetailsView->SetObject(CurrentProfileSettingsObject);
+	}
+}
+
+void SCustomLaunchProfilesPanel::OnProfileModified(const ILauncherProfilePtr& Profile)
+{
+	if (Model->GetSelectedProfile() == Profile && Profile.IsValid())
+	{
+		TArray<TWeakObjectPtr<UObject>> SelectedObjects = DetailsView->GetSelectedObjects();
+		if (SelectedObjects.Num() == 1)
+		{
+			UBaseLauncherProfileSettings* ProfileSettings = Cast<UBaseLauncherProfileSettings>(SelectedObjects[0].Get());
+			if (ProfileSettings != nullptr)
+			{
+				ProfileSettings->RefreshProfile();
+			}
+		}
+	}
+}
+
+
+
+void SCustomLaunchProfilesPanel::OnProfileLaunchComplete()
+{
+	OutputLog->RefreshLog();
+}
+
+
+void SCustomLaunchProfilesPanel::OnLogAreaExpansionChanged( bool bExpanded )
+{
+	LogAreaSplitter->SlotAt(1).SetSizingRule( bExpanded ? SSplitter::ESizeRule::FractionOfParent : SSplitter::ESizeRule::SizeToContent );
+
+	LogAreaSplitter->SlotAt(0).SetResizable(bExpanded);
+	LogAreaSplitter->SlotAt(1).SetResizable(bExpanded);
+}
+
+const FSlateBrush* SCustomLaunchProfilesPanel::GetSelectedProfileImage() const
+{
+	if (Model->GetSelectedProfile().IsValid() && Model->IsAdvancedProfile(Model->GetSelectedProfile().ToSharedRef()))
+	{
+		return FAppStyle::GetBrush("Icons.Warning.Large");
+	}
+	else
+	{
+		const PlatformInfo::FTargetPlatformInfo* PlatformInfo = Model->GetPlatformInfo(Model->GetSelectedProfile());
+		return FCustomLaunchUIStyle::GetProfileBrushForPlatform(PlatformInfo, EPlatformIconSize::Large);
+	}
+}
+
+bool SCustomLaunchProfilesPanel::IsSelectedProfileReadOnly() const
+{
+	return !Model->GetSelectedProfile().IsValid() || Model->IsBasicLaunchProfile(Model->GetSelectedProfile()); 
+}
+
+FText SCustomLaunchProfilesPanel::GetSelectedProfileName() const
+{
+	if (Model->IsBasicLaunchProfile(Model->GetSelectedProfile()))
+	{
+		const TSharedPtr<ITargetDeviceProxy> DeviceProxy = Model->GetDeviceProxy(Model->GetSelectedProfile().ToSharedRef());
+		if (DeviceProxy.IsValid())
+		{
+			return FText::FromString(DeviceProxy->GetName());
+		}
+		else
+		{
+			return LOCTEXT("NoDevice", "No Device");
+		}
+	}
+	else if (Model->GetSelectedProfile().IsValid())
+	{
+		return FText::FromString(Model->GetSelectedProfile()->GetName());
+	}
+	
+	return FText::GetEmpty();
+}
+
+
+
+FText SCustomLaunchProfilesPanel::GetSelectedProfileDescription() const
+{
+	if (Model->IsBasicLaunchProfile(Model->GetSelectedProfile()))
+	{
+		const TSharedPtr<ITargetDeviceProxy> DeviceProxy = Model->GetDeviceProxy(Model->GetSelectedProfile().ToSharedRef());
+		if (DeviceProxy.IsValid())
+		{
+			return DeviceProxy->GetPlatformDisplayName(NAME_None);
+		}
+		else
+		{
+			return FText::GetEmpty();
+		}
+	}
+	else if (Model->GetSelectedProfile().IsValid())
+	{
+		return FText::FromString(Model->GetSelectedProfile()->GetDescription());
+	}
+
+	return FText::GetEmpty();
+}
+
+
+
+
+
+void SCustomLaunchProfilesPanel::SetProfileEditorVisibleCheckState(const ECheckBoxState NewCheckState)
+{
+	SetProfileEditorVisible( NewCheckState == ECheckBoxState::Checked );
+}
+
+ECheckBoxState SCustomLaunchProfilesPanel::GetProfileEditorVisibleCheckState() const
+{
+	return IsProfileEditorVisible() ? ECheckBoxState::Checked : ECheckBoxState::Unchecked;
+}
+
+void SCustomLaunchProfilesPanel::SetProfileEditorVisible( bool bVisible )
+{
+	if (ProfileEditorSplitter.IsValid())
+	{
+		ProfileEditorSplitter->SlotAt(1).SetSizingRule( bVisible ? SSplitter::ESizeRule::FractionOfParent : SSplitter::ESizeRule::SizeToContent );
+	
+		ProfileEditorSplitter->SlotAt(0).SetResizable(bVisible);
+		ProfileEditorSplitter->SlotAt(1).SetResizable(bVisible);
+	}
+}
+
+bool SCustomLaunchProfilesPanel::IsProfileEditorVisible() const
+{
+	return !ProfileEditorSplitter.IsValid() || ProfileEditorSplitter->SlotAt(1).CanBeResized();
+
+}
+
+
+void SCustomLaunchProfilesPanel::SetDefaultProjectPath(FString ProjectPath)
+{
+	Model->GetProfileManager()->SetProjectPath(ProjectPath);
+}
+
+FString SCustomLaunchProfilesPanel::GetDefaultProjectPath() const
+{
+	return Model->GetProfileManager()->GetProjectPath();
+}
+
+void SCustomLaunchProfilesPanel::SetDefaultBuildTarget(FString BuildTarget)
+{
+	Model->GetProfileManager()->SetBuildTarget(BuildTarget);
+}
+
+FString SCustomLaunchProfilesPanel::GetDefaultBuildTarget() const
+{
+	return Model->GetProfileManager()->GetBuildTarget();
+}
+
+
+
+#undef LOCTEXT_NAMESPACE
diff --git a/Engine/Plugins/Developer/CustomLaunchUI/Source/Private/Widgets/Profiles/SCustomLaunchProfilesPanel.h b/Engine/Plugins/Developer/CustomLaunchUI/Source/Private/Widgets/Profiles/SCustomLaunchProfilesPanel.h
new file mode 100644
--- /dev/null
+++ b/Engine/Plugins/Developer/CustomLaunchUI/Source/Private/Widgets/Profiles/SCustomLaunchProfilesPanel.h
@@ -0,0 +1,90 @@
+// Copyright Epic Games, Inc. All Rights Reserved.
+
+#pragma once
+
+#include "CoreTypes.h"
+#include "Widgets/SCompoundWidget.h"
+#include "Model/CustomLaunchUIModel.h"
+#include "ILauncherWorker.h"
+#include "UObject/ObjectPtr.h"
+
+class SSplitter;
+class SExpandableArea;
+class SCustomLaunchCustomProfileSelector;
+
+class SCustomLaunchProfilesPanel : public SCompoundWidget
+{
+public:
+	SLATE_BEGIN_ARGS(SCustomLaunchProfilesPanel) {}
+		SLATE_EVENT(CustomLaunchUI::FOnProfileClicked, OnProfileLaunchClicked)
+	SLATE_END_ARGS()
+
+public:
+	SCustomLaunchProfilesPanel();
+	~SCustomLaunchProfilesPanel();
+
+public:
+	void Construct(const FArguments& InArgs, const TSharedRef<CustomLaunchUI::FModel>& InModel);
+
+	void OnProfileLaunchComplete();
+
+private:
+	TSharedRef<SWidget> CreateProfilesPanel();
+
+	TSharedRef<SWidget> CreateProfileSelectorWidget();
+	TSharedRef<SWidget> CreateProfileEditorWidget();
+	TSharedRef<SWidget> CreateAdvancedProfileWarningWidget();
+	TSharedRef<SWidget> CreateProfileEditorToolbarWidget();
+
+
+	void OnProfileEdit(const ILauncherProfilePtr& Profile);
+	void OnProfileDelete(const ILauncherProfilePtr& Profile);
+	void OnProfileSave(const ILauncherProfilePtr& Profile);
+	void OnProfileDuplicate(const ILauncherProfilePtr& Profile);
+	void OnProfileRename(const ILauncherProfilePtr& Profile);
+	void OnProfileSelected(const ILauncherProfilePtr& NewProfile, const ILauncherProfilePtr& OldProfile);
+	void OnProfileModified(const ILauncherProfilePtr& Profile);
+
+	void OnLogAreaExpansionChanged( bool bExpanded );
+
+	FReply OnLaunchButtonClicked();
+	bool IsLaunchButtonEnabled() const;
+	FText GetLaunchButtonToolTipText() const;
+	const FSlateBrush* GetLaunchButtonImage() const;
+
+	FReply OnCreateNewCustomProfileClicked();
+
+	FReply OnOpenDeviceManagerClicked();
+
+	const FSlateBrush* GetSelectedProfileImage() const;
+	bool IsSelectedProfileReadOnly() const;
+
+	FText GetSelectedProfileName() const;
+	FText GetSelectedProfileDescription() const;
+
+	void SetProfileEditorVisibleCheckState(const ECheckBoxState NewCheckState);
+	ECheckBoxState GetProfileEditorVisibleCheckState() const;
+
+	void SetProfileEditorVisible(bool bVisible);
+	bool IsProfileEditorVisible() const;
+
+	void SetDefaultProjectPath(FString ProjectPath);
+	FString GetDefaultProjectPath() const;
+
+	void SetDefaultBuildTarget(FString BuildTarget);
+	FString GetDefaultBuildTarget() const;
+
+
+	CustomLaunchUI::FOnProfileClicked OnProfileLaunchClicked;
+	TSharedPtr<CustomLaunchUI::FModel> Model;
+
+	TSharedPtr<SSplitter> LogAreaSplitter;
+	TSharedPtr<SExpandableArea> LogExpandableArea;
+	TSharedPtr<SSplitter> ProfileEditorSplitter;
+	TSharedPtr<SCustomLaunchCustomProfileSelector> ProfileSelector;
+	TSharedPtr<class IDetailsView> DetailsView;
+	TObjectPtr<class UBaseLauncherProfileSettings> CurrentProfileSettingsObject;
+	TSharedPtr<class SCustomLaunchOutputLog> OutputLog;
+
+
+};
diff --git a/Engine/Plugins/Developer/CustomLaunchUI/Source/Private/Widgets/SCustomLaunchUI.cpp b/Engine/Plugins/Developer/CustomLaunchUI/Source/Private/Widgets/SCustomLaunchUI.cpp
new file mode 100644
--- /dev/null
+++ b/Engine/Plugins/Developer/CustomLaunchUI/Source/Private/Widgets/SCustomLaunchUI.cpp
@@ -0,0 +1,116 @@
+// Copyright Epic Games, Inc. All Rights Reserved.
+
+#include "Widgets/SCustomLaunchUI.h"
+#include "SlateOptMacros.h"
+#include "Styling/StyleColors.h"
+#include "Widgets/Layout/SWidgetSwitcher.h"
+#include "Widgets/Layout/SBorder.h"
+#include "Widgets/Profiles/SCustomLaunchProfilesPanel.h"
+#include "Widgets/Output/SCustomLaunchLaunchPanel.h"
+#include "ILauncher.h"
+
+
+#define LOCTEXT_NAMESPACE "SCustomLaunchUI"
+
+
+SCustomLaunchUI::SCustomLaunchUI()
+{
+}
+
+
+SCustomLaunchUI::~SCustomLaunchUI()
+{
+	if (LauncherWorker.IsValid())
+	{
+		LauncherWorker->Cancel();
+		FPlatformProcess::Sleep(0.5f);
+	}
+}
+
+
+BEGIN_SLATE_FUNCTION_BUILD_OPTIMIZATION
+void SCustomLaunchUI::Construct(const FArguments& InArgs, const TSharedRef<CustomLaunchUI::FModel>& InModel)
+{
+	Model = InModel;
+
+	ChildSlot
+	[
+		SNew(SBorder)
+		.BorderImage(FAppStyle::GetBrush("ToolPanel.GroupBorder"))
+		.Padding(0)
+		[
+			SAssignNew(LauncherPanelWidgetSwitcher, SWidgetSwitcher)
+			.WidgetIndex(Panel_Profiles)
+		
+			// Profiles panel
+			+SWidgetSwitcher::Slot()
+			[
+				SAssignNew(ProfileSelectorPanel, SCustomLaunchProfilesPanel, Model.ToSharedRef())
+				.OnProfileLaunchClicked(this, &SCustomLaunchUI::OnProfileRun)
+			]
+
+			// Progress Panel
+			+SWidgetSwitcher::Slot()
+			[
+				SAssignNew(LaunchProgressPanel, SCustomLaunchLaunchPanel, Model.ToSharedRef())
+				.OnCloseClicked(this, &SCustomLaunchUI::OnProgressClose)
+				.OnRerunClicked(this, &SCustomLaunchUI::OnRerunClicked)
+			]
+		]
+	];
+}
+END_SLATE_FUNCTION_BUILD_OPTIMIZATION
+
+
+void SCustomLaunchUI::OnProfileRun(const ILauncherProfilePtr& Profile)
+{
+	LauncherProfile = Profile;
+	if (LauncherProfile.IsValid())
+	{
+		LauncherWorker = Model->GetLauncher()->Launch(Model->GetDeviceProxyManager(), LauncherProfile.ToSharedRef());
+	
+		if (LauncherWorker.IsValid())
+		{
+			LaunchProgressPanel->SetLauncherWorker(LauncherWorker.ToSharedRef());
+			LauncherPanelWidgetSwitcher->SetActiveWidgetIndex(Panel_Progress);
+		}
+	}
+}
+
+
+
+FReply SCustomLaunchUI::OnProgressClose()
+{
+	if (LauncherWorker.IsValid())
+	{
+		LauncherWorker->Cancel();
+	}
+	LauncherProfile.Reset();
+
+	LauncherPanelWidgetSwitcher->SetActiveWidgetIndex(Panel_Profiles);
+	ProfileSelectorPanel->OnProfileLaunchComplete();
+
+	return FReply::Handled();
+}
+
+
+
+FReply SCustomLaunchUI::OnRerunClicked()
+{
+	if (LauncherWorker.IsValid())
+	{
+		LauncherWorker->Cancel();
+	}
+	LauncherWorker = Model->GetLauncher()->Launch(Model->GetDeviceProxyManager(), LauncherProfile.ToSharedRef());
+
+	if (LauncherWorker.IsValid())
+	{
+		LaunchProgressPanel->SetLauncherWorker(LauncherWorker.ToSharedRef());
+	}
+
+	return FReply::Handled();
+}
+
+
+
+#undef LOCTEXT_NAMESPACE
diff --git a/Engine/Plugins/Developer/CustomLaunchUI/Source/Private/Widgets/SCustomLaunchUI.h b/Engine/Plugins/Developer/CustomLaunchUI/Source/Private/Widgets/SCustomLaunchUI.h
new file mode 100644
--- /dev/null
+++ b/Engine/Plugins/Developer/CustomLaunchUI/Source/Private/Widgets/SCustomLaunchUI.h
@@ -0,0 +1,45 @@
+// Copyright Epic Games, Inc. All Rights Reserved.
+
+#pragma once
+
+#include "CoreTypes.h"
+#include "Widgets/SCompoundWidget.h"
+#include "Model/CustomLaunchUIModel.h"
+#include "ILauncherWorker.h"
+
+class SWidgetSwitcher;
+class SCustomLaunchLaunchPanel;
+class SCustomLaunchProfilesPanel;
+
+class SCustomLaunchUI : public SCompoundWidget
+{
+public:
+	SLATE_BEGIN_ARGS(SCustomLaunchUI) { }
+	SLATE_END_ARGS()
+
+public:
+	SCustomLaunchUI();
+	~SCustomLaunchUI();
+
+public:
+	void Construct(const FArguments& InArgs, const TSharedRef<CustomLaunchUI::FModel>& InModel);
+
+private:
+	void OnProfileRun(const ILauncherProfilePtr& Profile);
+	FReply OnProgressClose();
+	FReply OnRerunClicked();
+
+	ILauncherWorkerPtr LauncherWorker;
+	ILauncherProfilePtr LauncherProfile;
+	TSharedPtr<CustomLaunchUI::FModel> Model;
+
+	enum EPanel
+	{
+		Panel_Profiles = 0,
+		Panel_Progress = 1,
+	};
+
+	TSharedPtr<SWidgetSwitcher> LauncherPanelWidgetSwitcher;
+	TSharedPtr<SCustomLaunchProfilesPanel> ProfileSelectorPanel;
+	TSharedPtr<SCustomLaunchLaunchPanel> LaunchProgressPanel;
+};
diff --git a/Engine/Plugins/Developer/CustomLaunchUI/Source/Private/Widgets/SSegmentedProgressBar.cpp b/Engine/Plugins/Developer/CustomLaunchUI/Source/Private/Widgets/SSegmentedProgressBar.cpp
new file mode 100644
--- /dev/null
+++ b/Engine/Plugins/Developer/CustomLaunchUI/Source/Private/Widgets/SSegmentedProgressBar.cpp
@@ -0,0 +1,361 @@
+// Copyright Epic Games, Inc. All Rights Reserved.
+
+#include "SSegmentedProgressBar.h"
+#include "Widgets/Images/SImage.h"
+#include "Widgets/SBoxPanel.h"
+#include "Widgets/SOverlay.h"
+#include "Widgets/Colors/SColorBlock.h"
+#include "Styling/CustomLaunchUIStyle.h"
+#include "Styling/StyleColors.h"
+#include "Layout/LayoutUtils.h"
+
+#define LOCTEXT_NAMESPACE "SSegmentedProgressBar"
+
+SSegmentedProgressBar::FScopedWidgetSlotArguments SSegmentedProgressBar::InsertSlot(int32 Index, bool bRebuildChildren)
+{
+	if (bRebuildChildren)
+	{
+		TWeakPtr<SSegmentedProgressBar> AsWeak = SharedThis(this);
+		return FScopedWidgetSlotArguments { MakeUnique<FSlot>(), this->Children, Index, [AsWeak](const FSlot*, int32)
+		{
+			if (TSharedPtr<SSegmentedProgressBar> SharedThis = AsWeak.Pin())
+			{
+				SharedThis->RebuildChildren();
+			}
+		}};
+	}
+	else
+	{
+		return FScopedWidgetSlotArguments(MakeUnique<FSlot>(), this->Children, Index);
+	}
+}
+
+
+
+void SSegmentedProgressBar::Construct( const FArguments& InArgs )
+{
+	ThrobberAnimation = FCurveSequence(0.0f, 1.0f);
+	ThrobberAnimation.Play(AsShared(), true);
+
+	Children.AddSlots(MoveTemp(const_cast<TArray<typename FSlot::FSlotArguments>&>(InArgs._Slots)));	
+	RebuildChildren();
+}
+
+
+
+void SSegmentedProgressBar::Tick(const FGeometry& AllottedGeometry, const double InCurrentTime, const float InDeltaTime)
+{
+	SCompoundWidget::Tick(AllottedGeometry, InCurrentTime, InDeltaTime);
+	CacheChildStates();
+}
+
+
+
+void SSegmentedProgressBar::CacheChildStates()
+{
+	CachedChildStates.Reset();
+	for (int32 ChildIndex = 0; ChildIndex < Children.Num(); ChildIndex++)
+	{
+		EState ChildState = Children[ChildIndex].State.Get();
+
+		// show us as busy as soon as the previous task is completed - just looks nicer
+		if (ChildState == EState::Pending && (ChildIndex == 0 || CachedChildStates[ChildIndex-1] == EState::Completed))
+		{
+			ChildState = EState::Busy;
+		}
+
+		CachedChildStates.Add( ChildState );
+	}
+
+	if (CachedChildStates.Num() == 0 )
+	{
+		CachedOverallState = EState::None;
+	}
+	else
+	{
+		CachedOverallState = CachedChildStates.Last();
+	}
+}
+
+
+
+SSegmentedProgressBar::EState SSegmentedProgressBar::GetChildState( int32 ChildIndex ) const
+{
+	if (CachedChildStates.IsValidIndex(ChildIndex))
+	{
+		return CachedChildStates[ChildIndex];
+	}
+	
+	return EState::None;
+}
+
+
+
+void SSegmentedProgressBar::RebuildChildren()
+{
+	CacheChildStates();
+
+	TSharedPtr<SHorizontalBox> SlotBox;
+	ChildSlot
+	[
+		SAssignNew(SlotBox, SHorizontalBox)
+	];
+	
+
+	// construct child widgets as needed
+	const int32 NumChildren = Children.Num();
+	for ( int32 ChildIndex = 0; ChildIndex < NumChildren; ChildIndex++ )
+	{
+		TSharedRef<SWidget> Child = Children[ChildIndex].GetWidget();
+		FSlot* ChildSlotPtr = &Children[ChildIndex];
+
+		if (Child == SNullWidget::NullWidget)
+		{
+			Child = ConstructChild(*ChildSlotPtr, ChildIndex);
+		}
+
+
+		// separator bar between previous and current
+		if (ChildIndex > 0)
+		{
+			AddChildSeparatorBar(SlotBox, ChildIndex);
+		}
+
+		// task item
+		SlotBox->AddSlot()
+		.VAlign(VAlign_Center)
+		.HAlign(HAlign_Center)
+		.AutoWidth()
+		[
+			Child
+		];
+	}
+}
+
+
+
+void SSegmentedProgressBar::AddChildSeparatorBar( TSharedPtr<SHorizontalBox> SlotBox, int32 ChildIndex ) const
+{
+	SlotBox->AddSlot()
+	.VAlign(VAlign_Center)
+	.HAlign(HAlign_Fill)
+	.FillWidth(1)
+	[
+		SNew(SColorBlock)
+		.Color( this, &SSegmentedProgressBar::GetSeparatorBarColor, ChildIndex )
+		.Size( FVector2D(32, LineSize) )
+	];
+}
+
+
+
+TSharedRef<SWidget> SSegmentedProgressBar::ConstructChild( const FSlot& Slot, int32 ChildIndex ) const
+{
+	return SNew(SOverlay)
+
+	// full circle (only shown when fully complete or cancelled)
+	+SOverlay::Slot()
+	.VAlign(VAlign_Center)
+	.HAlign(HAlign_Center)
+	[
+		SNew(SImage)
+		.Image( FCustomLaunchUIStyle::Get().GetBrush("FullCircle") )
+		.ColorAndOpacity( this, &SSegmentedProgressBar::GetCircleColor, ChildIndex)
+		.Visibility( this, &SSegmentedProgressBar::GetFullCircleVisibility, ChildIndex )
+
+	]
+
+	// outer circle
+	+SOverlay::Slot()
+	.VAlign(VAlign_Center)
+	.HAlign(HAlign_Center)
+	[
+		SNew(SImage)
+		.DesiredSizeOverride(FVector2D(36,36))
+		.Image( FCustomLaunchUIStyle::Get().GetBrush("OuterCircle") )
+		.ColorAndOpacity( this, &SSegmentedProgressBar::GetCircleColor, ChildIndex )
+		.Visibility( this, &SSegmentedProgressBar::GetOuterCircleVisibility, ChildIndex )
+	]
+
+	// outer busy circle
+	+SOverlay::Slot()
+	.VAlign(VAlign_Center)
+	.HAlign(HAlign_Center)
+	[
+		SNew(SImage)
+		.DesiredSizeOverride(FVector2D(36,36))
+		.Image( FCustomLaunchUIStyle::Get().GetBrush("OuterCircle.Busy") )
+		.ColorAndOpacity( FCustomLaunchUIStyle::Get().GetSlateColor("State.Busy") )
+		.Visibility( this, &SSegmentedProgressBar::GetProgressCircleVisibility, ChildIndex )
+		.RenderTransform( this, &SSegmentedProgressBar::GetProgressCircleTransform, ChildIndex)
+		.RenderTransformPivot(FVector2D(.5f,.5f))
+	]
+
+	// task icon
+	+SOverlay::Slot()
+	.VAlign(VAlign_Center)
+	.HAlign(HAlign_Center)
+	[
+		SNew(SImage)
+		.Image(Slot.Image)
+		.ColorAndOpacity(this, &SSegmentedProgressBar::GetIconColor, ChildIndex)
+		.DesiredSizeOverride(FVector2D(16,16))
+		.ToolTipText(Slot.ToolTipText)
+	]
+
+	// task state overlay
+	+SOverlay::Slot()
+	.VAlign(VAlign_Bottom)
+	.HAlign(HAlign_Right)
+	[
+		SNew(SImage)
+		.Image(this, &SSegmentedProgressBar::GetOverlayIcon, ChildIndex)
+		.Visibility( this, &SSegmentedProgressBar::GetOverlayVisibility, ChildIndex)
+	]
+	;
+}
+
+
+FLinearColor SSegmentedProgressBar::GetSeparatorBarColor( int32 ChildIndex ) const
+{
+	if (CachedOverallState == EState::Completed)
+	{
+		return FCustomLaunchUIStyle::Get().GetSlateColor("State.AllComplete").GetSpecifiedColor();
+	}
+
+	EState State = GetChildState(ChildIndex);
+	switch (State)
+	{
+		case EState::Busy:		return FCustomLaunchUIStyle::Get().GetSlateColor("State.Busy").GetSpecifiedColor();
+		case EState::Canceled:  return FCustomLaunchUIStyle::Get().GetSlateColor("State.Canceled").GetSpecifiedColor();
+		case EState::Completed:
+		case EState::Failed:	return FCustomLaunchUIStyle::Get().GetSlateColor("State.Complete").GetSpecifiedColor(); // line indicates previous step succeeded in this case
+		case EState::Pending:	return FCustomLaunchUIStyle::Get().GetSlateColor("State.Pending").GetSpecifiedColor();
+	}
+
+	return FSlateColor::UseForeground().GetSpecifiedColor();
+}
+
+FSlateColor SSegmentedProgressBar::GetCircleColor( int32 ChildIndex ) const
+{
+	if (CachedOverallState == EState::Completed)
+	{
+		return FCustomLaunchUIStyle::Get().GetSlateColor("State.AllComplete");
+	}
+
+	EState State = GetChildState(ChildIndex);
+	switch (State)
+	{
+		case EState::Busy:		return FCustomLaunchUIStyle::Get().GetSlateColor("State.Pending");
+		case EState::Canceled:  return FCustomLaunchUIStyle::Get().GetSlateColor("State.Canceled");
+		case EState::Completed:	return FCustomLaunchUIStyle::Get().GetSlateColor("State.Complete");
+		case EState::Failed:	return FCustomLaunchUIStyle::Get().GetSlateColor("State.Error");
+		case EState::Pending:	return FCustomLaunchUIStyle::Get().GetSlateColor("State.Pending");
+	}
+
+	return FSlateColor::UseForeground();
+}
+
+
+
+
+
+FSlateColor SSegmentedProgressBar::GetIconColor( int32 ChildIndex ) const
+{
+	if (CachedOverallState == EState::Completed)
+	{
+		return FStyleColors::White;
+	}
+
+	EState State = GetChildState(ChildIndex);
+	if (State == EState::Canceled)
+	{
+		return FStyleColors::Hover2;
+	}
+
+	return FStyleColors::Foreground;
+}
+
+
+EVisibility SSegmentedProgressBar::GetFullCircleVisibility( int32 ChildIndex ) const
+{
+	if (CachedOverallState == EState::Completed || GetChildState(ChildIndex) == EState::Canceled)
+	{
+		return EVisibility::Visible;
+	}
+
+	return EVisibility::Collapsed;
+}
+
+
+
+TOptional<FSlateRenderTransform> SSegmentedProgressBar::GetProgressCircleTransform( int32 ChildIndex ) const
+{
+	if (GetChildState(ChildIndex) == EState::Busy)
+	{
+		const float DeltaAngle = ThrobberAnimation.GetLerp()*2*PI;
+		return FSlateRenderTransform(FQuat2D(DeltaAngle));
+	}
+
+	FSlateRenderTransform Result;
+	return Result;
+}
+
+
+
+EVisibility SSegmentedProgressBar::GetOuterCircleVisibility( int32 ChildIndex ) const
+{
+	if (CachedOverallState != EState::Completed && GetChildState(ChildIndex) != EState::Canceled)
+	{
+		return EVisibility::Visible;
+	}
+
+	return EVisibility::Collapsed;
+}
+
+
+
+EVisibility SSegmentedProgressBar::GetProgressCircleVisibility( int32 ChildIndex ) const
+{
+	if (CachedOverallState != EState::Completed && GetChildState(ChildIndex) == EState::Busy)
+	{
+		return EVisibility::Visible;
+	}
+
+	return EVisibility::Collapsed;
+}
+
+
+
+
+EVisibility SSegmentedProgressBar::GetOverlayVisibility( int32 ChildIndex ) const
+{
+	EState State = GetChildState(ChildIndex);
+	if (State == EState::Completed || State == EState::Failed)
+	{
+		return EVisibility::Visible;
+	}
+
+	return EVisibility::Collapsed;
+}
+
+
+
+const FSlateBrush* SSegmentedProgressBar::GetOverlayIcon( int32 ChildIndex ) const
+{
+	if (CachedOverallState == EState::Completed)
+	{
+		return FCustomLaunchUIStyle::Get().GetBrush("BadgeOutlined.AllComplete");
+	}
+
+	EState State = GetChildState(ChildIndex);
+	switch (State)
+	{
+		case EState::Completed: return FCustomLaunchUIStyle::Get().GetBrush("BadgeOutlined.Success");
+		case EState::Failed:	return FCustomLaunchUIStyle::Get().GetBrush("BadgeOutlined.Error");
+	}
+
+	return FStyleDefaults::GetNoBrush();
+}
+
+#undef LOCTEXT_NAMESPACE
\ No newline at end of file
diff --git a/Engine/Plugins/Developer/CustomLaunchUI/Source/Private/Widgets/SSegmentedProgressBar.h b/Engine/Plugins/Developer/CustomLaunchUI/Source/Private/Widgets/SSegmentedProgressBar.h
new file mode 100644
--- /dev/null
+++ b/Engine/Plugins/Developer/CustomLaunchUI/Source/Private/Widgets/SSegmentedProgressBar.h
@@ -0,0 +1,137 @@
+// Copyright Epic Games, Inc. All Rights Reserved.
+
+#pragma once
+
+#include "CoreMinimal.h"
+#include "Styling/SlateColor.h"
+#include "Styling/CoreStyle.h"
+#include "Widgets/SCompoundWidget.h"
+#include "Animation/CurveSequence.h"
+
+
+class SSegmentedProgressBar : public SCompoundWidget
+{
+public:
+	enum class EState : uint8
+	{
+		None,
+		Busy,
+		Canceled,
+		Completed,
+		Failed,
+		Pending,
+	};
+
+
+	struct FSlot : public TSlotBase<FSlot>
+	{
+	public:
+		FSlot()
+			: TSlotBase<FSlot>()
+			, Image(nullptr)
+		{}
+
+		SLATE_SLOT_BEGIN_ARGS(FSlot, TSlotBase<FSlot>)
+			SLATE_ATTRIBUTE(const FSlateBrush*, Image)
+			SLATE_ATTRIBUTE(EState, State)
+			SLATE_ATTRIBUTE(FText, ToolTipText)
+		SLATE_SLOT_END_ARGS()
+
+		void Construct(const FChildren& SlotOwner, FSlotArguments&& InArgs)
+		{
+			TSlotBase<FSlot>::Construct(SlotOwner, MoveTemp(InArgs));
+			if (InArgs._Image.IsSet())
+			{
+				Image = MoveTemp(InArgs._Image);
+			}
+			if (InArgs._State.IsSet())
+			{
+				State = MoveTemp(InArgs._State);
+			}
+			if (InArgs._ToolTipText.IsSet())
+			{
+				ToolTipText = MoveTemp(InArgs._ToolTipText);
+			}
+		}
+
+	protected:
+		TAttribute<const FSlateBrush*> Image;
+		TAttribute<EState> State;
+		TAttribute<FText> ToolTipText;
+
+		friend SSegmentedProgressBar;
+	};
+
+	static FSlot::FSlotArguments Slot()
+	{
+		return FSlot::FSlotArguments(MakeUnique<FSlot>());
+	}
+
+	SLATE_BEGIN_ARGS(SSegmentedProgressBar){}
+		SLATE_SLOT_ARGUMENT(FSlot, Slots)
+	SLATE_END_ARGS()
+
+	using FScopedWidgetSlotArguments = typename TPanelChildren<FSlot>::FScopedWidgetSlotArguments;
+	FScopedWidgetSlotArguments AddSlot(bool bRebuildChildren = true)
+	{
+		return InsertSlot(INDEX_NONE, bRebuildChildren);
+	}
+
+	FScopedWidgetSlotArguments InsertSlot(int32 Index = INDEX_NONE, bool bRebuildChildren = true);
+
+	FSlot& GetSlot(int32 SlotIndex)
+	{
+		return Children[SlotIndex];
+	}
+
+	const FSlot& GetSlot(int32 SlotIndex) const
+	{
+		return Children[SlotIndex];
+	}
+
+	int32 NumSlots() const
+	{
+		return Children.Num();
+	}
+
+	void ClearChildren()
+	{
+		Children.Empty();
+	}
+
+	SSegmentedProgressBar()
+		: Children(this)
+	{
+		bCanSupportFocus = false;
+	}
+
+
+
+	void Construct(const FArguments& Arguments);
+	virtual void Tick(const FGeometry& AllottedGeometry, const double InCurrentTime, const float InDeltaTime) override;
+
+protected:
+	void RebuildChildren();
+	void AddChildSeparatorBar( TSharedPtr<class SHorizontalBox> SlotBox, int32 ChildIndex ) const;
+	TSharedRef<SWidget> ConstructChild( const FSlot& Slot, int32 ChildIndex ) const;
+
+	void CacheChildStates();
+	EState GetChildState( int32 ChildIndex ) const;
+
+	FLinearColor GetSeparatorBarColor( int32 ChildIndex ) const;
+	FSlateColor GetCircleColor( int32 ChildIndex ) const;
+	FSlateColor GetIconColor( int32 ChildIndex ) const;
+	EVisibility GetFullCircleVisibility( int32 ChildIndex ) const;
+	EVisibility GetOuterCircleVisibility( int32 ChildIndex ) const;
+	EVisibility GetProgressCircleVisibility( int32 ChildIndex ) const;
+	EVisibility GetOverlayVisibility( int32 ChildIndex ) const;
+	const FSlateBrush* GetOverlayIcon( int32 ChildIndex ) const;
+	TOptional<FSlateRenderTransform> GetProgressCircleTransform( int32 ChildIndex ) const;
+
+	TPanelChildren<FSlot> Children;
+	TArray<EState> CachedChildStates;
+	EState CachedOverallState = EState::None;
+
+	FCurveSequence ThrobberAnimation;
+	const float LineSize = 4;
+};
diff --git a/Engine/Plugins/Developer/CustomLaunchUI/Source/Private/Widgets/Shared/SCustomLaunchBuildTargetCombo.cpp b/Engine/Plugins/Developer/CustomLaunchUI/Source/Private/Widgets/Shared/SCustomLaunchBuildTargetCombo.cpp
new file mode 100644
--- /dev/null
+++ b/Engine/Plugins/Developer/CustomLaunchUI/Source/Private/Widgets/Shared/SCustomLaunchBuildTargetCombo.cpp
@@ -0,0 +1,143 @@
+// Copyright Epic Games, Inc. All Rights Reserved.
+
+#include "Widgets/Shared/SCustomLaunchBuildTargetCombo.h"
+
+#include "SlateOptMacros.h"
+#include "Widgets/Input/SComboButton.h"
+#include "Framework/MultiBox/MultiBoxBuilder.h"
+#include "Interfaces/IMainFrameModule.h"
+#include "DesktopPlatformModule.h"
+#include "Misc/Paths.h"
+#include "Settings/ProjectPackagingSettings.h"
+
+#define LOCTEXT_NAMESPACE "SCustomLaunchBuildTargetCombo"
+
+
+BEGIN_SLATE_FUNCTION_BUILD_OPTIMIZATION
+void SCustomLaunchBuildTargetCombo::Construct(const FArguments& InArgs)
+{
+	OnSelectionChanged = InArgs._OnSelectionChanged;
+	SelectedProject = InArgs._SelectedProject;
+	SelectedBuildTarget = InArgs._SelectedBuildTarget;
+	SupportedTargetTypes = InArgs._SupportedTargetTypes;
+
+	FSlateFontInfo Font = InArgs._Font.IsSet() ? InArgs._Font.Get() : InArgs._TextStyle->Font;
+
+	ChildSlot
+	[
+		SNew(SComboButton)
+		.ButtonContent()
+		[
+			SNew(STextBlock)
+			.Text(this, &SCustomLaunchBuildTargetCombo::GetBuildTargetName)
+			.Font(Font)
+		]
+		.OnGetMenuContent(this, &SCustomLaunchBuildTargetCombo::MakeBuildTargetSelectionWidget)
+	];
+}
+END_SLATE_FUNCTION_BUILD_OPTIMIZATION
+
+
+
+BEGIN_SLATE_FUNCTION_BUILD_OPTIMIZATION
+TSharedRef<SWidget> SCustomLaunchBuildTargetCombo::MakeBuildTargetSelectionWidget()
+{
+	FMenuBuilder MenuBuilder(true, nullptr);
+	{
+		FUIAction AnyProjectAction(FExecuteAction::CreateSP(this, &SCustomLaunchBuildTargetCombo::SetBuildTargetName, FString()));
+
+		FString DefaultBuildTarget = GetDefaultBuildTargetName();
+		FText DefaultBuildTargetName = DefaultBuildTarget.IsEmpty() ? LOCTEXT("DefaultBuildTargetName", "Project Default") : FText::Format(LOCTEXT("CurProjectDefaultBuildTargetName", "{0} (Project Default)"), FText::FromString(DefaultBuildTarget));
+		MenuBuilder.AddMenuEntry(DefaultBuildTargetName, LOCTEXT("DefaultBuildTargetActionHint", "Use the project default build target."), FSlateIcon(), AnyProjectAction);
+
+		if (SelectedProject.IsSet())
+		{
+			MenuBuilder.AddMenuSeparator();
+
+			const TArray<FTargetInfo>& BuildTargets = FDesktopPlatformModule::Get()->GetTargetsForProject(SelectedProject.Get());
+			for (const FTargetInfo& BuildTarget : BuildTargets)
+			{
+				if (BuildTarget.Type == EBuildTargetType::Program)
+				{
+					continue;
+				}
+
+				if (!SupportedTargetTypes.IsSet() || SupportedTargetTypes.Get().IsEmpty() || SupportedTargetTypes.Get().Contains(BuildTarget.Type))
+				{
+					FUIAction ProjectAction(FExecuteAction::CreateSP(this, &SCustomLaunchBuildTargetCombo::SetBuildTargetName, BuildTarget.Name));
+					MenuBuilder.AddMenuEntry(FText::FromString(BuildTarget.Name), FText::FromString(BuildTarget.Path), FSlateIcon(), ProjectAction);
+				}
+				else
+				{
+					FUIAction ProjectAction(FExecuteAction::CreateSP(this, &SCustomLaunchBuildTargetCombo::SetBuildTargetName, BuildTarget.Name),
+											FCanExecuteAction::CreateLambda([]() { return false; }));
+					MenuBuilder.AddMenuEntry(FText::FromString(BuildTarget.Name), LOCTEXT("NotCompatible", "The selected platform does not support this build target type"), FSlateIcon(), ProjectAction);
+				}
+			}
+		}
+	}
+
+	return MenuBuilder.MakeWidget();
+}
+END_SLATE_FUNCTION_BUILD_OPTIMIZATION
+
+
+FText SCustomLaunchBuildTargetCombo::GetBuildTargetName() const
+{
+	FString BuildTargetName = SelectedBuildTarget.Get();
+
+	if (BuildTargetName.IsEmpty())
+	{
+		FString DefaultBuildTarget = GetDefaultBuildTargetName();
+		if (DefaultBuildTarget.IsEmpty())
+		{
+			return LOCTEXT("DefaultBuildTargetName", "Project Default");
+		}
+		else
+		{
+			return FText::Format(LOCTEXT("CurProjectDefaultBuildTargetName", "{0} (Project Default)"), FText::FromString(DefaultBuildTarget));
+		}
+	}
+
+	return FText::FromString(BuildTargetName);
+}
+
+void SCustomLaunchBuildTargetCombo::SetBuildTargetName(FString BuildTargetName)
+{
+	OnSelectionChanged.ExecuteIfBound(BuildTargetName);
+}
+
+FString SCustomLaunchBuildTargetCombo::GetDefaultBuildTargetName() const
+{
+	if (SelectedProject.IsSet())
+	{
+		// if this is the current in-editor project, we can check the default setting
+		if (FPaths::IsProjectFilePathSet() && SelectedProject.Get() == FPaths::GetProjectFilePath())
+		{
+			UProjectPackagingSettings* AllPlatformPackagingSettings = GetMutableDefault<UProjectPackagingSettings>();
+			if (AllPlatformPackagingSettings != nullptr && !AllPlatformPackagingSettings->BuildTarget.IsEmpty())
+			{
+				return AllPlatformPackagingSettings->BuildTarget;
+			}
+		}
+
+		// if there is only one game build target, it will be used by default
+		const TArray<FTargetInfo>& BuildTargets = FDesktopPlatformModule::Get()->GetTargetsForProject(SelectedProject.Get());
+		const TArray<FTargetInfo> GameBuildTargets = BuildTargets.FilterByPredicate([](const FTargetInfo& TargetInfo)
+		{
+			return TargetInfo.Type == EBuildTargetType::Game;
+		});
+		if (GameBuildTargets.Num() == 1)
+		{
+			return GameBuildTargets[0].Name;
+		}
+
+		// multiple (or no) game targets... would need to read the BuildTarget from the other project's ini file
+		// ... fixme ...
+	}
+
+	return FString();
+}
+
+
+#undef LOCTEXT_NAMESPACE
diff --git a/Engine/Plugins/Developer/CustomLaunchUI/Source/Private/Widgets/Shared/SCustomLaunchBuildTargetCombo.h b/Engine/Plugins/Developer/CustomLaunchUI/Source/Private/Widgets/Shared/SCustomLaunchBuildTargetCombo.h
new file mode 100644
--- /dev/null
+++ b/Engine/Plugins/Developer/CustomLaunchUI/Source/Private/Widgets/Shared/SCustomLaunchBuildTargetCombo.h
@@ -0,0 +1,47 @@
+// Copyright Epic Games, Inc. All Rights Reserved.
+
+#pragma once
+
+#include "CoreMinimal.h"
+#include "Widgets/SCompoundWidget.h"
+#include "Styling/SlateTypes.h"
+#include "Styling/CoreStyle.h"
+#include "Styling/SlateWidgetStyleAsset.h"
+#include "Fonts/SlateFontInfo.h"
+
+class SCustomLaunchBuildTargetCombo
+	: public SCompoundWidget
+{
+public:
+	DECLARE_DELEGATE_OneParam(FOnSelectionChanged, FString );
+
+	SLATE_BEGIN_ARGS(SCustomLaunchBuildTargetCombo)
+		: _TextStyle( &FCoreStyle::Get().GetWidgetStyle<FTextBlockStyle>( "NormalText" ) )
+		{}
+		SLATE_EVENT(FOnSelectionChanged, OnSelectionChanged);
+		SLATE_ATTRIBUTE(FString, SelectedProject);
+		SLATE_ATTRIBUTE(FString, SelectedBuildTarget);
+		SLATE_ATTRIBUTE(TArray<EBuildTargetType>, SupportedTargetTypes);
+		SLATE_ATTRIBUTE(FSlateFontInfo, Font)
+		SLATE_STYLE_ARGUMENT( FTextBlockStyle, TextStyle )
+	SLATE_END_ARGS()
+
+public:
+	void Construct(	const FArguments& InArgs);
+
+protected:
+	FString GetDefaultBuildTargetName() const;
+
+	TAttribute<TArray<EBuildTargetType>> SupportedTargetTypes;
+	TAttribute<FString> SelectedProject;
+	TAttribute<FString> SelectedBuildTarget;
+	FOnSelectionChanged OnSelectionChanged;
+	bool bShowAnyProjectOption;
+
+	TSharedRef<SWidget> MakeBuildTargetSelectionWidget();
+
+	FText GetBuildTargetName() const;
+	void SetBuildTargetName(FString BuildTargetName);
+
+
+};
diff --git a/Engine/Plugins/Developer/CustomLaunchUI/Source/Private/Widgets/Shared/SCustomLaunchContentSchemeCombo.cpp b/Engine/Plugins/Developer/CustomLaunchUI/Source/Private/Widgets/Shared/SCustomLaunchContentSchemeCombo.cpp
new file mode 100644
--- /dev/null
+++ b/Engine/Plugins/Developer/CustomLaunchUI/Source/Private/Widgets/Shared/SCustomLaunchContentSchemeCombo.cpp
@@ -0,0 +1,81 @@
+// Copyright Epic Games, Inc. All Rights Reserved.
+
+#include "Widgets/Shared/SCustomLaunchContentSchemeCombo.h"
+
+#include "SlateOptMacros.h"
+#include "Widgets/Input/SComboButton.h"
+#include "Widgets/Text/STextBlock.h"
+#include "Framework/MultiBox/MultiBoxBuilder.h"
+
+#define LOCTEXT_NAMESPACE "SCustomLaunchContentSchemeCombo"
+
+
+BEGIN_SLATE_FUNCTION_BUILD_OPTIMIZATION
+void SCustomLaunchContentSchemeCombo::Construct(const FArguments& InArgs)
+{
+	OnSelectionChanged = InArgs._OnSelectionChanged;
+	IsContentSchemeAvailable = InArgs._IsContentSchemeAvailable;
+	SelectedContentScheme = InArgs._SelectedContentScheme;
+
+	FSlateFontInfo Font = InArgs._Font.IsSet() ? InArgs._Font.Get() : InArgs._TextStyle->Font;
+
+	ChildSlot
+	[
+		SNew(SComboButton)
+		.ButtonContent()
+		[
+			SNew(STextBlock)
+			.Text(this, &SCustomLaunchContentSchemeCombo::GetContentSchemeName)
+			.Font(Font)
+		]
+		.OnGetMenuContent(this, &SCustomLaunchContentSchemeCombo::MakeContentSchemeSelectionWidget)
+	];
+}
+END_SLATE_FUNCTION_BUILD_OPTIMIZATION
+
+
+
+BEGIN_SLATE_FUNCTION_BUILD_OPTIMIZATION
+TSharedRef<SWidget> SCustomLaunchContentSchemeCombo::MakeContentSchemeSelectionWidget()
+{
+	FMenuBuilder MenuBuilder(true, nullptr);
+	{
+		for (CustomLaunchUI::EContentScheme ContentScheme : CustomLaunchUI::GetAllContentSchemes())
+		{
+			if (IsContentSchemeAvailable.IsBound() && !IsContentSchemeAvailable.Execute(ContentScheme))
+			{
+				continue;
+			}
+
+			MenuBuilder.AddMenuEntry(
+				CustomLaunchUI::GetContentSchemeDisplayName(ContentScheme),
+				CustomLaunchUI::GetContentSchemeToolTip(ContentScheme),
+				FSlateIcon(), 
+				FUIAction(
+					FExecuteAction::CreateSP(this, &SCustomLaunchContentSchemeCombo::SetContentScheme, ContentScheme),
+					FCanExecuteAction(),
+					FGetActionCheckState::CreateLambda( [this, ContentScheme]() { return (SelectedContentScheme.Get() == ContentScheme) ? ECheckBoxState::Checked : ECheckBoxState::Unchecked; })
+				),
+				NAME_None,
+				EUserInterfaceActionType::Check);
+		}
+	}
+
+	return MenuBuilder.MakeWidget();
+}
+END_SLATE_FUNCTION_BUILD_OPTIMIZATION
+
+
+FText SCustomLaunchContentSchemeCombo::GetContentSchemeName() const
+{
+	CustomLaunchUI::EContentScheme ContentScheme = SelectedContentScheme.Get();
+	return CustomLaunchUI::GetContentSchemeDisplayName(ContentScheme);
+}
+
+void SCustomLaunchContentSchemeCombo::SetContentScheme(CustomLaunchUI::EContentScheme ContentScheme)
+{
+	OnSelectionChanged.ExecuteIfBound(ContentScheme);
+}
+
+
+#undef LOCTEXT_NAMESPACE
diff --git a/Engine/Plugins/Developer/CustomLaunchUI/Source/Private/Widgets/Shared/SCustomLaunchContentSchemeCombo.h b/Engine/Plugins/Developer/CustomLaunchUI/Source/Private/Widgets/Shared/SCustomLaunchContentSchemeCombo.h
new file mode 100644
--- /dev/null
+++ b/Engine/Plugins/Developer/CustomLaunchUI/Source/Private/Widgets/Shared/SCustomLaunchContentSchemeCombo.h
@@ -0,0 +1,42 @@
+// Copyright Epic Games, Inc. All Rights Reserved.
+
+#pragma once
+
+#include "CoreMinimal.h"
+#include "Widgets/SCompoundWidget.h"
+#include "Styling/SlateTypes.h"
+#include "Styling/CoreStyle.h"
+#include "Styling/SlateWidgetStyleAsset.h"
+#include "Fonts/SlateFontInfo.h"
+#include "Model/CustomLaunchUIModel.h"
+
+class SCustomLaunchContentSchemeCombo
+	: public SCompoundWidget
+{
+public:
+	DECLARE_DELEGATE_OneParam(FOnSelectionChanged, CustomLaunchUI::EContentScheme );
+	DECLARE_DELEGATE_RetVal_OneParam(bool, FGetContentSchemeAvailability, CustomLaunchUI::EContentScheme )
+
+	SLATE_BEGIN_ARGS(SCustomLaunchContentSchemeCombo) 
+		: _TextStyle( &FCoreStyle::Get().GetWidgetStyle<FTextBlockStyle>( "NormalText" ) )
+		{}
+		SLATE_EVENT(FOnSelectionChanged, OnSelectionChanged);
+		SLATE_EVENT(FGetContentSchemeAvailability, IsContentSchemeAvailable)
+		SLATE_ATTRIBUTE(CustomLaunchUI::EContentScheme, SelectedContentScheme);
+		SLATE_ATTRIBUTE(FSlateFontInfo, Font)
+		SLATE_STYLE_ARGUMENT( FTextBlockStyle, TextStyle )
+	SLATE_END_ARGS()
+
+public:
+	void Construct(	const FArguments& InArgs);
+
+protected:
+	TAttribute<CustomLaunchUI::EContentScheme> SelectedContentScheme;
+	FOnSelectionChanged OnSelectionChanged;
+	FGetContentSchemeAvailability IsContentSchemeAvailable;
+
+	TSharedRef<SWidget> MakeContentSchemeSelectionWidget();
+
+	FText GetContentSchemeName() const;
+	void SetContentScheme(CustomLaunchUI::EContentScheme ContentScheme);
+};
diff --git a/Engine/Plugins/Developer/CustomLaunchUI/Source/Private/Widgets/Shared/SCustomLaunchDeviceCombo.cpp b/Engine/Plugins/Developer/CustomLaunchUI/Source/Private/Widgets/Shared/SCustomLaunchDeviceCombo.cpp
new file mode 100644
--- /dev/null
+++ b/Engine/Plugins/Developer/CustomLaunchUI/Source/Private/Widgets/Shared/SCustomLaunchDeviceCombo.cpp
@@ -0,0 +1,170 @@
+// Copyright Epic Games, Inc. All Rights Reserved.
+
+#include "Widgets/Shared/SCustomLaunchDeviceCombo.h"
+
+#include "Widgets/SBoxPanel.h"
+#include "Widgets/Input/SComboBox.h"
+#include "Widgets/Input/SCheckBox.h"
+#include "Widgets/Text/STextBlock.h"
+#include "Widgets/Images/SImage.h"
+#include "SlateOptMacros.h"
+#include "Styling/AppStyle.h"
+#include "Styling/CustomLaunchUIStyle.h"
+#include "ITargetDeviceProxy.h"
+#include "ITargetDeviceProxyManager.h"
+#include "PlatformInfo.h"
+
+#define LOCTEXT_NAMESPACE "SCustomLaunchDeviceCombo"
+
+
+BEGIN_SLATE_FUNCTION_BUILD_OPTIMIZATION
+void SCustomLaunchDeviceCombo::Construct(const FArguments& InArgs)
+{
+	OnDeviceRemoved = InArgs._OnDeviceRemoved;
+	OnSelectionChanged = InArgs._OnSelectionChanged;
+	SelectedDevices = InArgs._SelectedDevices;
+	Platforms = InArgs._Platforms;
+	bAllPlatforms = InArgs._AllPlatforms;
+
+	ChildSlot
+	[
+		SNew(SHorizontalBox)
+		
+		+SHorizontalBox::Slot()
+		.AutoWidth()		
+		[
+			SAssignNew(DeviceProxyComboBox, SComboBox<TSharedPtr<ITargetDeviceProxy>>)
+			.OptionsSource(&DeviceProxyList)
+			.OnGenerateWidget(this, &SCustomLaunchDeviceCombo::GenerateDeviceProxyListWidget)
+			.OnSelectionChanged( this, &SCustomLaunchDeviceCombo::OnDeviceProxySelectionChangedChanged )
+			[
+				SNew(SHorizontalBox)
+					
+				+SHorizontalBox::Slot()
+				.AutoWidth()
+				.VAlign(VAlign_Center)
+				.Padding(4,0)
+				[
+					SNew(SImage)
+					.DesiredSizeOverride(FVector2D(16,16))
+					.Image(this, &SCustomLaunchDeviceCombo::GetSelectedDeviceProxyBrush)
+				]
+
+				+SHorizontalBox::Slot()
+				.AutoWidth()
+				.VAlign(VAlign_Center)
+				.Padding(4,0)
+				[
+					SNew(STextBlock)
+					//.TextStyle( FAppStyle::Get(), "SmallText")
+					.Text(this, &SCustomLaunchDeviceCombo::GetSelectedDeviceProxyName)
+				]
+			]
+		]
+	];
+
+	SCustomLaunchDeviceWidgetBase::Construct();
+}
+END_SLATE_FUNCTION_BUILD_OPTIMIZATION
+
+
+
+
+
+
+BEGIN_SLATE_FUNCTION_BUILD_OPTIMIZATION
+TSharedRef<SWidget> SCustomLaunchDeviceCombo::GenerateDeviceProxyListWidget( TSharedPtr<ITargetDeviceProxy> DeviceProxy ) const
+{
+	const PlatformInfo::FTargetPlatformInfo* PlatformInfo = PlatformInfo::FindPlatformInfo(*DeviceProxy->GetTargetPlatformName(NAME_None));
+
+	return SNew(SHorizontalBox)
+
+		// platform icon
+		+SHorizontalBox::Slot()
+		.AutoWidth()
+		.VAlign(VAlign_Center)
+		.Padding(4,0)
+		[
+			SNew(SImage)
+			.DesiredSizeOverride(FVector2D(16,16))
+			.Image(PlatformInfo ? FAppStyle::GetBrush(PlatformInfo->GetIconStyleName(EPlatformIconSize::Normal)) : FStyleDefaults::GetNoBrush())
+		]
+
+		// device name
+		+SHorizontalBox::Slot()
+		.AutoWidth()
+		.VAlign(VAlign_Center)
+		.Padding(4,0)
+		[
+			SNew(STextBlock)
+			.Text(FText::FromString(*DeviceProxy->GetName()))
+		]
+	;
+}
+END_SLATE_FUNCTION_BUILD_OPTIMIZATION
+
+
+const FSlateBrush* SCustomLaunchDeviceCombo::GetSelectedDeviceProxyBrush() const
+{
+	TArray<FString> DeviceIDs = SelectedDevices.Get();
+
+	if (DeviceIDs.Num() == 1)
+	{
+		FString DeviceID = DeviceIDs[0];
+		const TSharedPtr<ITargetDeviceProxy> DeviceProxy = GetDeviceProxyManager()->FindProxyDeviceForTargetDevice(DeviceID);
+		if (DeviceProxy.IsValid())
+		{
+			const PlatformInfo::FTargetPlatformInfo* PlatformInfo = PlatformInfo::FindPlatformInfo(*DeviceProxy->GetTargetPlatformName(NAME_None));
+			if (PlatformInfo != nullptr)
+			{
+				return FAppStyle::GetBrush(PlatformInfo->GetIconStyleName(EPlatformIconSize::Normal));
+			}
+		}
+	}
+	else if (DeviceIDs.Num() > 1)
+	{
+		return FAppStyle::Get().GetBrush("Icons.WarningWithColor");
+	}
+
+	return FStyleDefaults::GetNoBrush();
+}
+
+FText SCustomLaunchDeviceCombo::GetSelectedDeviceProxyName() const
+{
+	TArray<FString> DeviceIDs = SelectedDevices.Get();
+
+	if (DeviceIDs.Num() == 1)
+	{
+		FString DeviceID = DeviceIDs[0];
+		const TSharedPtr<ITargetDeviceProxy> DeviceProxy = GetDeviceProxyManager()->FindProxyDeviceForTargetDevice(DeviceID);
+		if (DeviceProxy.IsValid())
+		{
+			return FText::FromString(*DeviceProxy->GetName());
+		}
+	}
+	else if (DeviceIDs.Num() > 1)
+	{
+		return LOCTEXT("TooManyDevices", "Multiple devices (unsupported)");
+	}
+
+	return LOCTEXT("NoDevice", "(no device)");
+}
+
+
+
+
+
+
+void SCustomLaunchDeviceCombo::OnDeviceProxySelectionChangedChanged(TSharedPtr<ITargetDeviceProxy> DeviceProxy, ESelectInfo::Type InSelectInfo)
+{
+	FString DeviceID = DeviceProxy->GetTargetDeviceId(NAME_None);
+
+	TArray<FString> DeviceIDs;
+	DeviceIDs.Add( DeviceID );
+
+	OnSelectionChanged.ExecuteIfBound(DeviceIDs);
+}
+
+
+
+#undef LOCTEXT_NAMESPACE
diff --git a/Engine/Plugins/Developer/CustomLaunchUI/Source/Private/Widgets/Shared/SCustomLaunchDeviceCombo.h b/Engine/Plugins/Developer/CustomLaunchUI/Source/Private/Widgets/Shared/SCustomLaunchDeviceCombo.h
new file mode 100644
--- /dev/null
+++ b/Engine/Plugins/Developer/CustomLaunchUI/Source/Private/Widgets/Shared/SCustomLaunchDeviceCombo.h
@@ -0,0 +1,37 @@
+// Copyright Epic Games, Inc. All Rights Reserved.
+
+#pragma once
+
+#include "CoreMinimal.h"
+#include "Widgets/Shared/SCustomLaunchDeviceWidgetBase.h"
+
+template<typename ItemType> class SComboBox;
+
+class SCustomLaunchDeviceCombo
+	: public SCustomLaunchDeviceWidgetBase
+{
+public:
+	SLATE_BEGIN_ARGS(SCustomLaunchDeviceCombo)
+		: _AllPlatforms(false)
+		{}
+		SLATE_EVENT(FOnDeviceRemoved, OnDeviceRemoved)
+		SLATE_EVENT(FOnSelectionChanged, OnSelectionChanged);
+		SLATE_ATTRIBUTE(TArray<FString>, SelectedDevices);
+		SLATE_ATTRIBUTE(TArray<FString>, Platforms);
+		SLATE_ARGUMENT(bool, AllPlatforms)
+	SLATE_END_ARGS()
+
+public:
+	void Construct(	const FArguments& InArgs );
+
+protected:
+
+	TSharedRef<SWidget> GenerateDeviceProxyListWidget(TSharedPtr<ITargetDeviceProxy> DeviceProxy) const;
+	void OnDeviceProxySelectionChangedChanged(TSharedPtr<ITargetDeviceProxy> DeviceProxy, ESelectInfo::Type InSelectInfo);
+	const FSlateBrush* GetSelectedDeviceProxyBrush() const;
+	FText GetSelectedDeviceProxyName() const;
+
+private:
+	TSharedPtr<SComboBox<TSharedPtr<ITargetDeviceProxy>> > DeviceProxyComboBox;
+
+};
diff --git a/Engine/Plugins/Developer/CustomLaunchUI/Source/Private/Widgets/Shared/SCustomLaunchDeviceListView.cpp b/Engine/Plugins/Developer/CustomLaunchUI/Source/Private/Widgets/Shared/SCustomLaunchDeviceListView.cpp
new file mode 100644
--- /dev/null
+++ b/Engine/Plugins/Developer/CustomLaunchUI/Source/Private/Widgets/Shared/SCustomLaunchDeviceListView.cpp
@@ -0,0 +1,145 @@
+// Copyright Epic Games, Inc. All Rights Reserved.
+
+#include "Widgets/Shared/SCustomLaunchDeviceListView.h"
+
+#include "Widgets/SBoxPanel.h"
+#include "Widgets/Views/SListView.h"
+#include "Widgets/Input/SCheckBox.h"
+#include "Widgets/Text/STextBlock.h"
+#include "Widgets/Images/SImage.h"
+#include "Widgets/Views/STableRow.h"
+#include "SlateOptMacros.h"
+#include "Styling/AppStyle.h"
+#include "Styling/CustomLaunchUIStyle.h"
+#include "ITargetDeviceProxy.h"
+#include "PlatformInfo.h"
+
+#define LOCTEXT_NAMESPACE "SCustomLaunchDeviceListView"
+
+
+BEGIN_SLATE_FUNCTION_BUILD_OPTIMIZATION
+void SCustomLaunchDeviceListView::Construct(const FArguments& InArgs)
+{
+	OnDeviceRemoved = InArgs._OnDeviceRemoved;
+	OnSelectionChanged = InArgs._OnSelectionChanged;
+	SelectedDevices = InArgs._SelectedDevices;
+	Platforms = InArgs._Platforms;
+	bAllPlatforms = InArgs._AllPlatforms;
+	bSingleSelect = InArgs._SingleSelect;
+
+	ChildSlot
+	[
+		SNew(SBorder)
+		.Padding(2)
+		.BorderImage(FAppStyle::GetBrush("ToolPanel.GroupBorder"))
+		[
+			SAssignNew(DeviceProxyListView, SListView<TSharedPtr<ITargetDeviceProxy>>)
+			.ListItemsSource(&DeviceProxyList)
+			.OnGenerateRow(this, &SCustomLaunchDeviceListView::GenerateDeviceProxyRow)
+		]
+	];
+
+	SCustomLaunchDeviceWidgetBase::Construct();
+}
+END_SLATE_FUNCTION_BUILD_OPTIMIZATION
+
+
+
+
+BEGIN_SLATE_FUNCTION_BUILD_OPTIMIZATION
+TSharedRef<ITableRow> SCustomLaunchDeviceListView::GenerateDeviceProxyRow(TSharedPtr<ITargetDeviceProxy> DeviceProxy, const TSharedRef<STableViewBase>& OwnerTable)
+{
+	const PlatformInfo::FTargetPlatformInfo* PlatformInfo = PlatformInfo::FindPlatformInfo(*DeviceProxy->GetTargetPlatformName(NAME_None));
+
+	TSharedPtr<SWidget> DeviceSelectorCheckbox = SNew(SCheckBox)
+	.IsChecked(this, &SCustomLaunchDeviceListView::IsDeviceProxyChecked, DeviceProxy)
+	.OnCheckStateChanged(this, &SCustomLaunchDeviceListView::OnDeviceProxyCheckStateChanged, DeviceProxy)
+	.Style(FAppStyle::Get(), bSingleSelect ? "RadioButton" : "Checkbox")
+	[
+		SNew(SHorizontalBox)
+
+		// platform icon
+		+SHorizontalBox::Slot()
+		.AutoWidth()
+		.VAlign(VAlign_Center)
+		.Padding(4,0)
+		[
+			SNew(SImage)
+			.DesiredSizeOverride(FVector2D(16,16))
+			.Image(PlatformInfo ? FAppStyle::GetBrush(PlatformInfo->GetIconStyleName(EPlatformIconSize::Normal)) : FStyleDefaults::GetNoBrush())
+		]
+
+		// device name
+		+SHorizontalBox::Slot()
+		.AutoWidth()
+		.VAlign(VAlign_Center)
+		.Padding(4,0)
+		[
+			SNew(STextBlock)
+			.Text(FText::FromString(*DeviceProxy->GetName()))
+		]
+	];
+
+	return SNew(STableRow<TSharedPtr<ITargetDeviceProxy>>, OwnerTable)
+	.Padding(FMargin(4,1))
+	[
+		DeviceSelectorCheckbox.ToSharedRef()
+	];
+}
+END_SLATE_FUNCTION_BUILD_OPTIMIZATION
+
+
+ECheckBoxState SCustomLaunchDeviceListView::IsDeviceProxyChecked(TSharedPtr<ITargetDeviceProxy> DeviceProxy) const
+{
+	FString DeviceID = DeviceProxy->GetTargetDeviceId(NAME_None);
+
+	if (SelectedDevices.Get().Contains(DeviceID))
+	{
+		return ECheckBoxState::Checked;
+	}
+	else
+	{
+		return ECheckBoxState::Unchecked;
+	}
+}
+
+
+
+void SCustomLaunchDeviceListView::OnDeviceProxyCheckStateChanged(ECheckBoxState NewState, TSharedPtr<ITargetDeviceProxy> DeviceProxy)
+{
+	FString DeviceID = DeviceProxy->GetTargetDeviceId(NAME_None);
+
+	TArray<FString> Devices;
+	if (!bSingleSelect)
+	{
+		Devices = SelectedDevices.Get();
+	}
+	else if (NewState != ECheckBoxState::Checked)
+	{
+		// do not allow the current item to be deselected in single-select mode
+		return;
+	}
+
+	if (NewState == ECheckBoxState::Checked)
+	{
+		Devices.Add(DeviceID);
+	}
+	else
+	{
+		Devices.Remove(DeviceID);
+	}
+
+	OnSelectionChanged.ExecuteIfBound(Devices);
+}
+
+
+void SCustomLaunchDeviceListView::OnDeviceListRefreshed()
+{
+	DeviceProxyListView->RequestListRefresh();
+}
+
+
+
+
+
+#undef LOCTEXT_NAMESPACE
diff --git a/Engine/Plugins/Developer/CustomLaunchUI/Source/Private/Widgets/Shared/SCustomLaunchDeviceListView.h b/Engine/Plugins/Developer/CustomLaunchUI/Source/Private/Widgets/Shared/SCustomLaunchDeviceListView.h
new file mode 100644
--- /dev/null
+++ b/Engine/Plugins/Developer/CustomLaunchUI/Source/Private/Widgets/Shared/SCustomLaunchDeviceListView.h
@@ -0,0 +1,42 @@
+// Copyright Epic Games, Inc. All Rights Reserved.
+
+#pragma once
+
+#include "CoreMinimal.h"
+#include "Widgets/Shared/SCustomLaunchDeviceWidgetBase.h"
+
+class ITableRow;
+class STableViewBase;
+template <typename ItemType> class SListView;
+
+class SCustomLaunchDeviceListView
+	: public SCustomLaunchDeviceWidgetBase
+{
+public:
+	SLATE_BEGIN_ARGS(SCustomLaunchDeviceListView)
+		: _AllPlatforms(false)
+		, _SingleSelect(false)
+		{}
+		SLATE_EVENT(FOnDeviceRemoved, OnDeviceRemoved)
+		SLATE_EVENT(FOnSelectionChanged, OnSelectionChanged);
+		SLATE_ATTRIBUTE(TArray<FString>, SelectedDevices);
+		SLATE_ATTRIBUTE(TArray<FString>, Platforms);
+		SLATE_ARGUMENT(bool, AllPlatforms)
+		SLATE_ARGUMENT(bool, SingleSelect)
+	SLATE_END_ARGS()
+
+public:
+	void Construct(	const FArguments& InArgs );
+
+protected:
+	bool bSingleSelect;
+
+	TSharedRef<ITableRow> GenerateDeviceProxyRow(TSharedPtr<ITargetDeviceProxy> DeviceProxy, const TSharedRef<STableViewBase>& OwnerTable);
+	ECheckBoxState IsDeviceProxyChecked(TSharedPtr<ITargetDeviceProxy> DeviceProxy) const;
+	void OnDeviceProxyCheckStateChanged(ECheckBoxState NewState, TSharedPtr<ITargetDeviceProxy> DeviceProxy);
+	virtual void OnDeviceListRefreshed() override;
+
+private:
+	TSharedPtr<SListView<TSharedPtr<ITargetDeviceProxy>> > DeviceProxyListView;
+
+};
diff --git a/Engine/Plugins/Developer/CustomLaunchUI/Source/Private/Widgets/Shared/SCustomLaunchDeviceWidgetBase.cpp b/Engine/Plugins/Developer/CustomLaunchUI/Source/Private/Widgets/Shared/SCustomLaunchDeviceWidgetBase.cpp
new file mode 100644
--- /dev/null
+++ b/Engine/Plugins/Developer/CustomLaunchUI/Source/Private/Widgets/Shared/SCustomLaunchDeviceWidgetBase.cpp
@@ -0,0 +1,128 @@
+// Copyright Epic Games, Inc. All Rights Reserved.
+
+#include "Widgets/Shared/SCustomLaunchDeviceWidgetBase.h"
+
+#include "ITargetDeviceProxy.h"
+#include "ITargetDeviceProxyManager.h"
+#include "ITargetDeviceServicesModule.h"
+#include "PlatformInfo.h"
+#include "Modules/ModuleManager.h"
+
+void SCustomLaunchDeviceWidgetBase::Construct()
+{
+	const TSharedRef<ITargetDeviceProxyManager> DeviceProxyManager = GetDeviceProxyManager();
+	DeviceProxyManager->OnProxyAdded().AddSP(this, &SCustomLaunchDeviceWidgetBase::OnDeviceProxyAdded);
+	DeviceProxyManager->OnProxyRemoved().AddSP(this, &SCustomLaunchDeviceWidgetBase::OnDeviceProxyRemoved);
+
+	RefreshDeviceList();
+}
+
+
+SCustomLaunchDeviceWidgetBase::~SCustomLaunchDeviceWidgetBase()
+{
+	const TSharedRef<ITargetDeviceProxyManager>& DeviceProxyManager = GetDeviceProxyManager();
+	DeviceProxyManager->OnProxyAdded().RemoveAll(this);
+	DeviceProxyManager->OnProxyRemoved().RemoveAll(this);
+}
+
+
+void SCustomLaunchDeviceWidgetBase::OnDeviceProxyAdded(const TSharedRef<ITargetDeviceProxy>& DeviceProxy)
+{
+	RefreshDeviceList();
+}
+
+
+
+void SCustomLaunchDeviceWidgetBase::OnDeviceProxyRemoved(const TSharedRef<ITargetDeviceProxy>& DeviceProxy)
+{
+	FString DeviceID = DeviceProxy->GetTargetDeviceId(NAME_None);
+	OnDeviceRemoved.ExecuteIfBound(DeviceID);
+
+	RefreshDeviceList();
+}
+
+
+
+void SCustomLaunchDeviceWidgetBase::RefreshDeviceList()
+{
+	DeviceProxyList.Reset();
+
+	if (bAllPlatforms)
+	{
+		for (const PlatformInfo::FTargetPlatformInfo* PlatformInfo : PlatformInfo::GetVanillaPlatformInfoArray())
+		{
+			TArray<TSharedPtr<ITargetDeviceProxy>> PlatformDeviceProxyList;
+			GetDeviceProxyManager()->GetProxies(PlatformInfo->Name, false, PlatformDeviceProxyList);
+
+			for (TSharedPtr<ITargetDeviceProxy> DeviceProxy : PlatformDeviceProxyList)
+			{
+				DeviceProxyList.AddUnique(DeviceProxy);
+			}
+		}
+	}
+	else
+	{
+		for (const FString& Platform : Platforms.Get())
+		{
+			const PlatformInfo::FTargetPlatformInfo* PlatformInfo = PlatformInfo::FindPlatformInfo(FName(Platform));
+			if (PlatformInfo != nullptr)
+			{
+				TArray<TSharedPtr<ITargetDeviceProxy>> PlatformDeviceProxyList;
+				GetDeviceProxyManager()->GetProxies(PlatformInfo->Name, false, PlatformDeviceProxyList);
+
+				for (TSharedPtr<ITargetDeviceProxy> DeviceProxy : PlatformDeviceProxyList)
+				{
+					DeviceProxyList.AddUnique(DeviceProxy);
+				}
+			}
+		}
+	}
+
+	OnDeviceListRefreshed();
+}
+
+
+void SCustomLaunchDeviceWidgetBase::OnSelectedPlatformChanged()
+{
+	// collect all device ids for the new platforms
+	TArray<FString> ValidPlatformDevices;
+	for (const FString& Platform : Platforms.Get())
+	{
+		const PlatformInfo::FTargetPlatformInfo* PlatformInfo = PlatformInfo::FindPlatformInfo(FName(Platform));
+		if (PlatformInfo != nullptr)
+		{
+			TArray<TSharedPtr<ITargetDeviceProxy>> PlatformDeviceProxyList;
+			GetDeviceProxyManager()->GetProxies(PlatformInfo->Name, false, PlatformDeviceProxyList);
+
+			for (TSharedPtr<ITargetDeviceProxy> PlatformDeviceProxy : PlatformDeviceProxyList)
+			{
+				ValidPlatformDevices.Add(PlatformDeviceProxy->GetTargetDeviceId(NAME_None));
+			}
+		}
+	}
+
+	// remove any devices that belong to other platforms
+	TArray<FString> ValidSelectedDevices = SelectedDevices.Get();
+	int32 InvalidCount = ValidSelectedDevices.RemoveAll( [ValidPlatformDevices]( const FString& Device )
+	{
+		return !ValidPlatformDevices.Contains(Device);
+	});
+
+	if (InvalidCount > 0)
+	{
+		OnSelectionChanged.ExecuteIfBound(ValidSelectedDevices);
+	}
+
+	// update the list
+	RefreshDeviceList();
+}
+
+
+const TSharedRef<ITargetDeviceProxyManager> SCustomLaunchDeviceWidgetBase::GetDeviceProxyManager() const
+{
+	ITargetDeviceServicesModule& TargetDeviceServicesModule = FModuleManager::LoadModuleChecked<ITargetDeviceServicesModule>("TargetDeviceServices");
+	return TargetDeviceServicesModule.GetDeviceProxyManager();
+}
+
+
+#undef LOCTEXT_NAMESPACE
diff --git a/Engine/Plugins/Developer/CustomLaunchUI/Source/Private/Widgets/Shared/SCustomLaunchDeviceWidgetBase.h b/Engine/Plugins/Developer/CustomLaunchUI/Source/Private/Widgets/Shared/SCustomLaunchDeviceWidgetBase.h
new file mode 100644
--- /dev/null
+++ b/Engine/Plugins/Developer/CustomLaunchUI/Source/Private/Widgets/Shared/SCustomLaunchDeviceWidgetBase.h
@@ -0,0 +1,40 @@
+// Copyright Epic Games, Inc. All Rights Reserved.
+
+#pragma once
+
+#include "CoreMinimal.h"
+#include "Widgets/SCompoundWidget.h"
+
+class ITargetDeviceProxy;
+class ITargetDeviceProxyManager;
+
+class SCustomLaunchDeviceWidgetBase
+	: public SCompoundWidget
+{
+public:
+	DECLARE_DELEGATE_OneParam(FOnSelectionChanged, TArray<FString> );
+	DECLARE_DELEGATE_OneParam(FOnDeviceRemoved, FString );
+
+public:
+	void Construct();
+	~SCustomLaunchDeviceWidgetBase();
+
+	void RefreshDeviceList();
+	void OnSelectedPlatformChanged();
+
+protected:
+	virtual void OnDeviceListRefreshed() {};
+
+	TAttribute<TArray<FString>> Platforms;
+	TAttribute<TArray<FString>> SelectedDevices;
+	FOnSelectionChanged OnSelectionChanged;
+	FOnDeviceRemoved OnDeviceRemoved;
+	bool bAllPlatforms = false;
+
+	void OnDeviceProxyAdded(const TSharedRef<ITargetDeviceProxy>& DeviceProxy);
+	void OnDeviceProxyRemoved(const TSharedRef<ITargetDeviceProxy>& DeviceProxy);
+
+	const TSharedRef<ITargetDeviceProxyManager> GetDeviceProxyManager() const;
+
+	TArray<TSharedPtr<ITargetDeviceProxy>> DeviceProxyList;
+};
diff --git a/Engine/Plugins/Developer/CustomLaunchUI/Source/Private/Widgets/Shared/SCustomLaunchMapListView.cpp b/Engine/Plugins/Developer/CustomLaunchUI/Source/Private/Widgets/Shared/SCustomLaunchMapListView.cpp
new file mode 100644
--- /dev/null
+++ b/Engine/Plugins/Developer/CustomLaunchUI/Source/Private/Widgets/Shared/SCustomLaunchMapListView.cpp
@@ -0,0 +1,398 @@
+// Copyright Epic Games, Inc. All Rights Reserved.
+
+#include "Widgets/Shared/SCustomLaunchMapListView.h"
+
+#include "Widgets/SBoxPanel.h"
+#include "Widgets/Views/STreeView.h"
+#include "Widgets/Input/SCheckBox.h"
+#include "Widgets/Input/SComboBox.h"
+#include "Widgets/Input/SSearchBox.h"
+#include "Widgets/Text/STextBlock.h"
+#include "Widgets/Images/SImage.h"
+#include "Widgets/Views/STableRow.h"
+#include "SlateOptMacros.h"
+#include "Styling/AppStyle.h"
+#include "Misc/PackageName.h"
+#include "HAL/FileManager.h"
+
+
+#define LOCTEXT_NAMESPACE "SCustomLaunchMapListView"
+
+
+BEGIN_SLATE_FUNCTION_BUILD_OPTIMIZATION
+void SCustomLaunchMapListView::Construct(const FArguments& InArgs)
+{
+	OnSelectionChanged = InArgs._OnSelectionChanged;
+	SelectedMaps = InArgs._SelectedMaps;
+	ProjectPath = InArgs._ProjectPath;
+
+	ChildSlot
+	[
+		SAssignNew(MapTreeView, STreeView<FMapTreeNodePtr>)
+		.TreeItemsSource(&MapTreeViewItemsSource)
+		.OnGenerateRow(this, &SCustomLaunchMapListView::GenerateMapTreeNodeRow)
+		.OnGetChildren(this, &SCustomLaunchMapListView::GetMapTreeNodeChildren)
+	];
+
+	RefreshMapList();
+}
+END_SLATE_FUNCTION_BUILD_OPTIMIZATION
+
+
+
+
+
+
+BEGIN_SLATE_FUNCTION_BUILD_OPTIMIZATION
+TSharedRef<SWidget> SCustomLaunchMapListView::MakeControlsWidget()
+{
+	return SNew(SHorizontalBox)
+		+SHorizontalBox::Slot()
+		.AutoWidth()
+		[
+			SNew(SCheckBox)
+			.Style( FAppStyle::Get(), "ToggleButtonCheckBox" )
+			.IsChecked_Lambda( [this]() { return bShowFolders ? ECheckBoxState::Checked : ECheckBoxState::Unchecked; })
+			.OnCheckStateChanged_Lambda( [this](ECheckBoxState CheckBoxState)
+			{
+				bShowFolders = CheckBoxState == ECheckBoxState::Checked;
+				RefreshMapList();
+			})
+			[
+				SNew(SImage)
+				.DesiredSizeOverride(FVector2D(16,16))
+				.Image( FAppStyle::GetBrush("Icons.FolderClosed")) // @todo: not the best image
+			]
+		]
+
+		+SHorizontalBox::Slot()
+		.FillWidth(1)
+		.Padding(4,0)
+		[
+			SNew(SSearchBox)
+			.OnTextCommitted(this, &SCustomLaunchMapListView::OnSearchFilterTextCommitted)
+			.OnTextChanged(this, &SCustomLaunchMapListView::OnSearchFilterTextChanged)
+		]
+	;
+}
+END_SLATE_FUNCTION_BUILD_OPTIMIZATION
+
+
+void SCustomLaunchMapListView::GetMapTreeNodeChildren(FMapTreeNodePtr Node, TArray<FMapTreeNodePtr>& OutChildren)
+{
+	OutChildren = Node->Children;
+}
+
+
+
+BEGIN_SLATE_FUNCTION_BUILD_OPTIMIZATION
+TSharedRef<ITableRow> SCustomLaunchMapListView::GenerateMapTreeNodeRow(FMapTreeNodePtr Node, const TSharedRef<STableViewBase>& OwnerTable)
+{
+	return SNew(STableRow<FMapTreeNodePtr>, OwnerTable)
+	[
+		SNew(SCheckBox)
+		.Padding(2)
+		.IsChecked(this, &SCustomLaunchMapListView::GetMapTreeNodeCheckState, Node)
+		.OnCheckStateChanged(this, &SCustomLaunchMapListView::SetMapTreeNodeCheckState, Node)
+		[
+			SNew(SHorizontalBox)
+
+			+SHorizontalBox::Slot()
+			.AutoWidth()
+			.VAlign(VAlign_Center)
+			.Padding(4,0)
+			[
+				SNew(SImage)
+				.DesiredSizeOverride(FVector2D(16,16))
+				.Image(this, &SCustomLaunchMapListView::GetMapTreeNodeIcon, Node)
+				.ColorAndOpacity(this, &SCustomLaunchMapListView::GetMapTreeNodeColor, Node)
+			]
+
+			+SHorizontalBox::Slot()
+			.AutoWidth()
+			.VAlign(VAlign_Center)
+			.Padding(4,0)
+			[
+				SNew(STextBlock)
+				.Text(FText::FromString(*Node->Name))
+			]
+		]
+	];
+}
+END_SLATE_FUNCTION_BUILD_OPTIMIZATION
+
+
+
+const FSlateBrush* SCustomLaunchMapListView::GetMapTreeNodeIcon(FMapTreeNodePtr Node) const
+{
+	if (Node->Children.Num() == 0)
+	{
+#if WITH_EDITOR
+		return FAppStyle::GetBrush("ContentBrowser.ColumnViewAssetIcon");
+#else
+		return FAppStyle::GetBrush("Icons.FilledCircle"); // nice asset icon not available outside of the editor... @todo: find a better one than this!
+#endif
+	}
+	else if (MapTreeView->IsItemExpanded(Node))
+	{
+		return FAppStyle::GetBrush("Icons.FolderOpen");
+	}
+	else
+	{
+		return FAppStyle::GetBrush("Icons.FolderClosed");
+	}
+}
+
+
+
+FSlateColor SCustomLaunchMapListView::GetMapTreeNodeColor(FMapTreeNodePtr Node) const
+{
+	if (Node->Children.Num() == 0)
+	{
+		return FAppStyle::Get().GetSlateColor("Colors.AccentOrange");
+
+	}
+	else
+	{
+		return FAppStyle::Get().GetSlateColor("Colors.AccentFolder");
+	}
+}
+
+
+
+ECheckBoxState SCustomLaunchMapListView::GetMapTreeNodeCheckState(FMapTreeNodePtr Node) const
+{
+	return Node->CheckBoxState;
+}
+
+
+
+void SCustomLaunchMapListView::SetMapTreeNodeCheckState(ECheckBoxState CheckBoxState, FMapTreeNodePtr Node)
+{
+	TArray<FString> CheckedMaps = SelectedMaps.Get();
+	SetCheckBoxStateRecursive(Node, CheckBoxState, CheckedMaps);
+
+	OnSelectionChanged.ExecuteIfBound(CheckedMaps);
+
+	RefreshCheckBoxState();
+}
+
+
+
+void SCustomLaunchMapListView::OnSearchFilterTextCommitted(const FText& SearchText, ETextCommit::Type CommitType)
+{
+	CurrentFilterText = SearchText.ToString();
+	RefreshMapList();
+}
+
+
+
+void SCustomLaunchMapListView::OnSearchFilterTextChanged(const FText& SearchText)
+{
+	CurrentFilterText = SearchText.ToString();
+	if (CurrentFilterText.IsEmpty())
+	{
+		RefreshMapList();
+	}
+}
+
+
+
+void SCustomLaunchMapListView::OnProjectChanged()
+{
+	RefreshMapList();
+}
+
+
+
+void SCustomLaunchMapListView::RefreshMapList()
+{
+	// helper to build the node tree
+	MapTreeRoot = MakeShared<FMapTreeNode>();
+	auto AddMapPath = [this]( const FString& Path )
+	{
+		// ignore this file if it isn't filtered
+		if (!CurrentFilterText.IsEmpty() && !FPaths::GetBaseFilename(Path).Contains(CurrentFilterText))
+		{
+			return;
+		}
+
+		if (bShowFolders)
+		{
+			TArray<FString> PathItems;
+			Path.ParseIntoArray(PathItems, TEXT("/"));
+
+			FMapTreeNodePtr Node = MapTreeRoot;
+			for (const FString& PathItem : PathItems)
+			{
+				if (FMapTreeNodePtr* ChildPtr = Node->Children.FindByPredicate( [PathItem]( const FMapTreeNodePtr Node ) { return Node->Name == PathItem; } ))
+				{
+					Node = *ChildPtr;
+				}
+				else
+				{
+					FMapTreeNodePtr Child = Node->Children.Add_GetRef(MakeShared<FMapTreeNode>());
+					Child->Name = FPaths::GetBaseFilename(PathItem);
+					Node = Child;
+				}
+			}
+		}
+		else
+		{
+			FMapTreeNodePtr Child = MapTreeRoot->Children.Add_GetRef(MakeShared<FMapTreeNode>());
+			Child->Name = FPaths::GetBaseFilename(Path);
+		}
+	};
+
+	// decide what to show
+	bool bShowProjectMaps = !ProjectPath.Get().IsEmpty();
+	bool bShowEngineMaps = false;
+
+	TArray<FString> AvailableMaps;
+	const FString WildCard = FString::Printf(TEXT("*%s"), *FPackageName::GetMapPackageExtension());
+
+	// show maps from the project
+	if (bShowProjectMaps)
+	{
+		FString ProjectBasePath = FPaths::GetPath(FPaths::ConvertRelativePathToFull(ProjectPath.Get()));
+		FString ProjectName = FPaths::GetBaseFilename(ProjectPath.Get());
+
+		FString ProjectContentDir = FPaths::Combine( ProjectBasePath, TEXT("Content") );
+		if (FPaths::IsRelative(ProjectContentDir))
+		{
+			ProjectContentDir = FPaths::Combine(FPaths::RootDir(), ProjectContentDir);
+		}
+
+		TArray<FString> ProjectMaps;
+		IFileManager::Get().FindFilesRecursive(ProjectMaps, *ProjectContentDir, *WildCard, true, false);
+		for (FString& ProjectMap : ProjectMaps)
+		{
+			AvailableMaps.AddUnique(FPaths::GetBaseFilename(ProjectMap));
+
+			ProjectMap.RemoveFromStart(*ProjectContentDir);
+			ProjectMap.ReplaceCharInline('\\', '/');
+			AddMapPath(ProjectName / ProjectMap);
+		}
+	}
+	else
+	{
+		bShowEngineMaps = true;
+	}
+	
+	
+	// show maps from the engine (fallback)
+	if (bShowEngineMaps)
+	{
+		FString EngineMapDir = FPaths::Combine(FPaths::EngineContentDir(), TEXT("Maps"));
+		TArray<FString> EngineMaps;
+		IFileManager::Get().FindFilesRecursive(EngineMaps, *EngineMapDir, *WildCard, true, false);
+		for (FString& EngineMap : EngineMaps)
+		{
+			AvailableMaps.AddUnique(FPaths::GetBaseFilename(EngineMap));
+
+			EngineMap.RemoveFromStart(*EngineMapDir);
+			EngineMap.ReplaceCharInline('\\', '/');
+			AddMapPath(TEXT("Engine") / EngineMap);
+		}
+	}
+
+	// add any maps that are selected but we have not found them
+	for (const FString& CookedMap : SelectedMaps.Get() )
+	{
+		if (!AvailableMaps.Contains(CookedMap))
+		{
+			AddMapPath(TEXT("Missing") / CookedMap);
+		}
+	}
+
+	RefreshCheckBoxState(true);
+
+	MapTreeViewItemsSource = MapTreeRoot->Children;
+	if (MapTreeView.IsValid())
+	{
+		MapTreeView->RequestTreeRefresh();
+	}
+
+}
+
+
+
+void SCustomLaunchMapListView::RefreshCheckBoxState( bool bExpand )
+{
+	RefreshCheckBoxStateRecursive(MapTreeRoot, bExpand);
+}
+
+
+
+ECheckBoxState SCustomLaunchMapListView::RefreshCheckBoxStateRecursive(FMapTreeNodePtr Node, bool bExpand)
+{
+	// update our filter state
+	Node->bFiltered = !CurrentFilterText.IsEmpty() && Node->Name.Contains(CurrentFilterText);
+
+	// this is a map not a folder - return the check state
+	if (Node->Children.Num() == 0)
+	{
+		return SelectedMaps.Get().Contains(Node->Name) ? ECheckBoxState::Checked : ECheckBoxState::Unchecked;
+	}
+
+	// this is a folder - check all children
+	ECheckBoxState LastBoxState = ECheckBoxState::Undetermined;
+	for (int32 Index = 0; Index < Node->Children.Num(); Index++)
+	{
+		// determine the child's check state
+		FMapTreeNodePtr Child = Node->Children[Index];
+		Child->CheckBoxState = RefreshCheckBoxStateRecursive(Child, bExpand);
+
+		// expand the child if necessary
+		if (bExpand && (Child->bFiltered || Child->CheckBoxState != ECheckBoxState::Unchecked))
+		{
+			MapTreeView->SetItemExpansion(Child, true);
+		}
+
+		// our check state becomes undetermined if the child check states do not match
+		if (Index > 0 && Child->CheckBoxState != LastBoxState)
+		{
+			LastBoxState = ECheckBoxState::Undetermined;
+		}
+		else
+		{
+			LastBoxState = Child->CheckBoxState;
+		}
+
+		// we are filtered if a child was filtered
+		Node->bFiltered |= Child->bFiltered;
+	}
+
+	// return this folder's check state
+	return LastBoxState;
+}
+
+
+
+void SCustomLaunchMapListView::SetCheckBoxStateRecursive(FMapTreeNodePtr Node, ECheckBoxState CheckBoxState, TArray<FString>& CheckedMaps)
+{
+	Node->CheckBoxState = CheckBoxState;
+
+	if (Node->Children.Num() == 0)
+	{
+		if (CheckBoxState == ECheckBoxState::Checked)
+		{
+			CheckedMaps.Add(Node->Name);
+		}
+		else
+		{
+			CheckedMaps.Remove(Node->Name);
+		}
+	}
+	else
+	{
+		for (FMapTreeNodePtr Child : Node->Children)
+		{
+			SetCheckBoxStateRecursive(Child, CheckBoxState, CheckedMaps);
+		}
+	}
+}
+
+
+
+#undef LOCTEXT_NAMESPACE
diff --git a/Engine/Plugins/Developer/CustomLaunchUI/Source/Private/Widgets/Shared/SCustomLaunchMapListView.h b/Engine/Plugins/Developer/CustomLaunchUI/Source/Private/Widgets/Shared/SCustomLaunchMapListView.h
new file mode 100644
--- /dev/null
+++ b/Engine/Plugins/Developer/CustomLaunchUI/Source/Private/Widgets/Shared/SCustomLaunchMapListView.h
@@ -0,0 +1,70 @@
+// Copyright Epic Games, Inc. All Rights Reserved.
+
+#pragma once
+
+#include "CoreMinimal.h"
+#include "Widgets/SCompoundWidget.h"
+
+class ITableRow;
+class STableViewBase;
+template<typename ItemType> class STreeView;
+
+class SCustomLaunchMapListView
+	: public SCompoundWidget
+{
+public:
+	DECLARE_DELEGATE_OneParam(FOnSelectionChanged, TArray<FString> );
+
+	SLATE_BEGIN_ARGS(SCustomLaunchMapListView){}
+		SLATE_EVENT(FOnSelectionChanged, OnSelectionChanged);
+		SLATE_ATTRIBUTE(TArray<FString>, SelectedMaps);
+		SLATE_ATTRIBUTE(FString, ProjectPath)
+	SLATE_END_ARGS()
+
+public:
+	void Construct(	const FArguments& InArgs );
+
+	void RefreshMapList();
+	TSharedRef<SWidget> MakeControlsWidget();
+
+protected:
+	TAttribute<TArray<FString>> SelectedMaps;
+	TAttribute<FString> ProjectPath;
+	FOnSelectionChanged OnSelectionChanged;
+
+private:
+	bool bShowFolders = true; // whether to display the available maps in a hierarchy or flat list
+
+	struct FMapTreeNode
+	{
+		FString Name;
+		ECheckBoxState CheckBoxState = ECheckBoxState::Unchecked;
+		bool bFiltered = false;
+		TArray<TSharedPtr<FMapTreeNode>> Children;
+	};
+	typedef TSharedPtr<FMapTreeNode> FMapTreeNodePtr;
+
+	FMapTreeNodePtr MapTreeRoot;
+
+	void OnProjectChanged();
+
+
+	void RefreshCheckBoxState( bool bExpand = false);
+	ECheckBoxState RefreshCheckBoxStateRecursive(FMapTreeNodePtr Node, bool bExpand);
+	void SetCheckBoxStateRecursive(FMapTreeNodePtr Node, ECheckBoxState CheckBoxState, TArray<FString>& CheckedMaps);
+
+	void OnSearchFilterTextCommitted(const FText& SearchText, ETextCommit::Type InCommitType);
+	void OnSearchFilterTextChanged(const FText& SearchText);
+
+	FString CurrentFilterText;
+
+	TSharedPtr<STreeView<FMapTreeNodePtr>> MapTreeView;
+	TArray<FMapTreeNodePtr> MapTreeViewItemsSource;
+
+	void GetMapTreeNodeChildren(FMapTreeNodePtr Item, TArray<FMapTreeNodePtr>& OutChildren);
+	TSharedRef<ITableRow> GenerateMapTreeNodeRow(FMapTreeNodePtr Item, const TSharedRef<STableViewBase>& OwnerTable);
+	ECheckBoxState GetMapTreeNodeCheckState(FMapTreeNodePtr Item) const;
+	void SetMapTreeNodeCheckState(ECheckBoxState CheckBoxState, FMapTreeNodePtr Item);
+	const FSlateBrush* GetMapTreeNodeIcon(FMapTreeNodePtr Node) const;
+	FSlateColor GetMapTreeNodeColor(FMapTreeNodePtr Node) const;
+};
diff --git a/Engine/Plugins/Developer/CustomLaunchUI/Source/Private/Widgets/Shared/SCustomLaunchPlatformCombo.cpp b/Engine/Plugins/Developer/CustomLaunchUI/Source/Private/Widgets/Shared/SCustomLaunchPlatformCombo.cpp
new file mode 100644
--- /dev/null
+++ b/Engine/Plugins/Developer/CustomLaunchUI/Source/Private/Widgets/Shared/SCustomLaunchPlatformCombo.cpp
@@ -0,0 +1,164 @@
+// Copyright Epic Games, Inc. All Rights Reserved.
+
+#include "Widgets/Shared/SCustomLaunchPlatformCombo.h"
+
+#include "SlateOptMacros.h"
+#include "Widgets/SBoxPanel.h"
+#include "Widgets/Input/SComboBox.h"
+#include "Widgets/Text/STextBlock.h"
+#include "Widgets/Images/SImage.h"
+#include "PlatformInfo.h"
+
+
+#define LOCTEXT_NAMESPACE "SCustomLaunchPlatformCombo"
+
+
+BEGIN_SLATE_FUNCTION_BUILD_OPTIMIZATION
+void SCustomLaunchPlatformCombo::Construct(const FArguments& InArgs)
+{
+	OnSelectionChanged = InArgs._OnSelectionChanged;
+	SelectedPlatforms = InArgs._SelectedPlatforms;
+	bBasicPlatformsOnly = InArgs._BasicPlatformsOnly;
+
+	ChildSlot
+	[
+		SAssignNew(PlatformsComboBox, SComboBox<TSharedPtr<FString>>)
+		.OptionsSource(&PlatformsList)
+		.OnGenerateWidget( this, &SCustomLaunchPlatformCombo::OnGeneratePlatformListWidget )
+		.OnSelectionChanged( this, &SCustomLaunchPlatformCombo::OnPlatformSelectionChanged )
+		[
+			SNew(SHorizontalBox)
+					
+			+SHorizontalBox::Slot()
+			.AutoWidth()
+			.VAlign(VAlign_Center)
+			.Padding(4,0)
+			[
+				SNew(SImage)
+				.DesiredSizeOverride(FVector2D(16,16))
+				.Image(this, &SCustomLaunchPlatformCombo::GetSelectedPlatformBrush)
+			]
+
+			+SHorizontalBox::Slot()
+			.AutoWidth()
+			.VAlign(VAlign_Center)
+			.Padding(4,0)
+			[
+				SNew(STextBlock)
+				.Text(this, &SCustomLaunchPlatformCombo::GetSelectedPlatformName)
+			]
+		]
+	];
+
+	for (const PlatformInfo::FTargetPlatformInfo* PlatformInfo : PlatformInfo::GetPlatformInfoArray())
+	{
+		if( bBasicPlatformsOnly && PlatformInfo->PlatformType != EBuildTargetType::Game)
+		{
+			continue;
+		}
+
+		PlatformsList.Add(MakeShared<FString>(PlatformInfo->Name.ToString()));
+	}
+
+}
+END_SLATE_FUNCTION_BUILD_OPTIMIZATION
+
+
+
+BEGIN_SLATE_FUNCTION_BUILD_OPTIMIZATION
+TSharedRef<SWidget> SCustomLaunchPlatformCombo::OnGeneratePlatformListWidget( TSharedPtr<FString> Platform ) const
+{
+	const PlatformInfo::FTargetPlatformInfo* PlatformInfo = PlatformInfo::FindPlatformInfo(FName(*Platform));
+	if (PlatformInfo != nullptr)
+	{
+		int32 Indent = PlatformInfo->IsVanilla() ? 0 : 16;
+
+		return SNew(SHorizontalBox)
+
+			+SHorizontalBox::Slot()
+			.AutoWidth()
+			.VAlign(VAlign_Center)
+			.Padding(4+Indent,4,4,4)
+			[
+				SNew(SImage)
+				.DesiredSizeOverride(FVector2D(16,16))
+				.Image(FAppStyle::GetBrush(PlatformInfo->GetIconStyleName(EPlatformIconSize::Normal)))
+			]
+
+			+SHorizontalBox::Slot()
+			.AutoWidth()
+			.VAlign(VAlign_Center)
+			.Padding(4)
+			[
+				SNew(STextBlock)
+				.Text(PlatformInfo->DisplayName)
+			]
+		;
+	}
+	else
+	{
+		return SNullWidget::NullWidget;
+	}
+}
+END_SLATE_FUNCTION_BUILD_OPTIMIZATION
+
+
+
+void SCustomLaunchPlatformCombo::OnPlatformSelectionChanged( TSharedPtr<FString> Platform, ESelectInfo::Type InSelectInfo )
+{
+	TArray<FString> Platforms;
+	if (Platform.IsValid() && !Platform->IsEmpty())
+	{
+		Platforms.Add(*Platform);
+	}
+
+	SelectedPlatforms.Set(Platforms);
+	OnSelectionChanged.ExecuteIfBound(SelectedPlatforms.Get());
+}
+
+
+
+const FSlateBrush* SCustomLaunchPlatformCombo::GetSelectedPlatformBrush() const
+{
+	const TArray<FString>& Platforms = SelectedPlatforms.Get();
+
+	if (Platforms.Num() == 1)
+	{
+		const PlatformInfo::FTargetPlatformInfo* PlatformInfo = PlatformInfo::FindPlatformInfo(FName(Platforms[0]));
+		if (PlatformInfo != nullptr)
+		{
+			return FAppStyle::GetBrush(PlatformInfo->GetIconStyleName(EPlatformIconSize::Normal));
+		}
+	}
+	else if (Platforms.Num() > 1)
+	{
+		return FAppStyle::Get().GetBrush("Icons.WarningWithColor");
+	}
+
+	return FStyleDefaults::GetNoBrush();
+}
+
+
+
+FText SCustomLaunchPlatformCombo::GetSelectedPlatformName() const
+{
+	const TArray<FString>& Platforms = SelectedPlatforms.Get();
+
+	if (Platforms.Num() == 1)
+	{
+		const PlatformInfo::FTargetPlatformInfo* PlatformInfo = PlatformInfo::FindPlatformInfo(FName(Platforms[0]));
+		if (PlatformInfo != nullptr)
+		{
+			return PlatformInfo->DisplayName;
+		}
+	}
+	else if (Platforms.Num() > 1)
+	{
+		return LOCTEXT("TooManyPlatforms", "Multiple platforms (unsupported)");
+	}
+
+	return LOCTEXT("NoPlatform", "(no platform)");
+}
+
+
+#undef LOCTEXT_NAMESPACE
diff --git a/Engine/Plugins/Developer/CustomLaunchUI/Source/Private/Widgets/Shared/SCustomLaunchPlatformCombo.h b/Engine/Plugins/Developer/CustomLaunchUI/Source/Private/Widgets/Shared/SCustomLaunchPlatformCombo.h
new file mode 100644
--- /dev/null
+++ b/Engine/Plugins/Developer/CustomLaunchUI/Source/Private/Widgets/Shared/SCustomLaunchPlatformCombo.h
@@ -0,0 +1,37 @@
+// Copyright Epic Games, Inc. All Rights Reserved.
+
+#pragma once
+
+#include "CoreMinimal.h"
+#include "Widgets/SCompoundWidget.h"
+
+template<typename ItemType> class SComboBox;
+
+
+class SCustomLaunchPlatformCombo
+	: public SCompoundWidget
+{
+public:
+	DECLARE_DELEGATE_OneParam(FOnSelectionChanged, TArray<FString> );
+
+	SLATE_BEGIN_ARGS(SCustomLaunchPlatformCombo){}
+		SLATE_EVENT(FOnSelectionChanged, OnSelectionChanged);
+		SLATE_ATTRIBUTE(TArray<FString>, SelectedPlatforms);
+		SLATE_ARGUMENT(bool, BasicPlatformsOnly)
+	SLATE_END_ARGS()
+
+public:
+	void Construct(	const FArguments& InArgs);
+
+protected:
+	TAttribute<TArray<FString>> SelectedPlatforms;
+	FOnSelectionChanged OnSelectionChanged;
+	bool bBasicPlatformsOnly;
+
+	TSharedRef<SWidget> OnGeneratePlatformListWidget( TSharedPtr<FString> Platform ) const;
+	void OnPlatformSelectionChanged( TSharedPtr<FString> Platform, ESelectInfo::Type InSelectInfo );
+	const FSlateBrush* GetSelectedPlatformBrush() const;
+	FText GetSelectedPlatformName() const;
+	TArray<TSharedPtr<FString>> PlatformsList;
+	TSharedPtr<SComboBox<TSharedPtr<FString>>> PlatformsComboBox;
+};
diff --git a/Engine/Plugins/Developer/CustomLaunchUI/Source/Private/Widgets/Shared/SCustomLaunchProjectCombo.cpp b/Engine/Plugins/Developer/CustomLaunchUI/Source/Private/Widgets/Shared/SCustomLaunchProjectCombo.cpp
new file mode 100644
--- /dev/null
+++ b/Engine/Plugins/Developer/CustomLaunchUI/Source/Private/Widgets/Shared/SCustomLaunchProjectCombo.cpp
@@ -0,0 +1,159 @@
+// Copyright Epic Games, Inc. All Rights Reserved.
+
+#include "Widgets/Shared/SCustomLaunchProjectCombo.h"
+
+#include "SlateOptMacros.h"
+#include "Widgets/Input/SComboButton.h"
+#include "Framework/Application/SlateApplication.h"
+#include "Framework/MultiBox/MultiBoxBuilder.h"
+#include "Interfaces/IMainFrameModule.h"
+#include "Model/CustomLaunchUIModel.h"
+#include "GameProjectHelper.h"
+#include "Misc/Paths.h"
+
+
+#define LOCTEXT_NAMESPACE "SCustomLaunchProjectCombo"
+
+
+BEGIN_SLATE_FUNCTION_BUILD_OPTIMIZATION
+void SCustomLaunchProjectCombo::Construct(const FArguments& InArgs)
+{
+	OnSelectionChanged = InArgs._OnSelectionChanged;
+	SelectedProject = InArgs._SelectedProject;
+	HasProject = InArgs._HasProject;
+	bShowAnyProjectOption = InArgs._ShowAnyProjectOption;
+	CurrentProjectOption = InArgs._CurrentProjectOption;
+
+	FSlateFontInfo Font = InArgs._Font.IsSet() ? InArgs._Font.Get() : InArgs._TextStyle->Font;
+
+	ChildSlot
+	[
+		SNew(SComboButton)
+		.ButtonContent()
+		[
+			SNew(STextBlock)
+			.Text(this, &SCustomLaunchProjectCombo::GetProjectName)
+			.Font(Font)
+		]
+		.MenuContent()
+		[
+			MakeProjectSelectionWidget()
+		]
+	];
+}
+END_SLATE_FUNCTION_BUILD_OPTIMIZATION
+
+
+
+BEGIN_SLATE_FUNCTION_BUILD_OPTIMIZATION
+TSharedRef<SWidget> SCustomLaunchProjectCombo::MakeProjectSelectionWidget()
+{
+	FMenuBuilder MenuBuilder(true, nullptr);
+	{
+		// any project option
+		if (bShowAnyProjectOption)
+		{
+			FUIAction AnyProjectAction(FExecuteAction::CreateSP(this, &SCustomLaunchProjectCombo::SetProjectPath, FString()));
+			MenuBuilder.AddMenuEntry(LOCTEXT("AnyProjectAction", "Any Project"), LOCTEXT("AnyProjectActionHint", "This profile can be used on any project. Build target selection will not be available."), FSlateIcon(), AnyProjectAction);
+			MenuBuilder.AddMenuSeparator();
+		}
+
+		// current project first
+		FString ProjectName;
+		if (CurrentProjectOption != ECurrentProjectOption::None)
+		{
+			FString ProjectPath;
+			if (FPaths::IsProjectFilePathSet())
+			{
+				ProjectName = FPaths::GetBaseFilename(FPaths::GetProjectFilePath());
+				if (CurrentProjectOption == ECurrentProjectOption::ActualProject)
+				{
+					ProjectPath = FPaths::GetProjectFilePath();
+				}
+			}
+
+			if (ProjectName.IsEmpty())
+			{
+				FUIAction CurrentProjectAction(FExecuteAction::CreateSP(this, &SCustomLaunchProjectCombo::SetProjectPath, FString()));
+				MenuBuilder.AddMenuEntry(LOCTEXT("CurrentProjectAction", "Current Project"), LOCTEXT("CurrentProjectHint", "Use the current project"), FSlateIcon(), CurrentProjectAction);
+			}
+			else if (CurrentProjectOption == ECurrentProjectOption::Empty)
+			{
+				FUIAction CurrentProjectAction(FExecuteAction::CreateSP(this, &SCustomLaunchProjectCombo::SetProjectPath, FString()));
+				MenuBuilder.AddMenuEntry(FText::Format(LOCTEXT("CurrentProjectActionFmt", "Current Project ({0})"), FText::FromString(ProjectName)), LOCTEXT("CurrentProjectHint", "Use the current project"), FSlateIcon(), CurrentProjectAction);
+			}
+			else
+			{
+				FUIAction CurrentProjectAction(FExecuteAction::CreateSP(this, &SCustomLaunchProjectCombo::SetProjectPath, ProjectPath));
+				MenuBuilder.AddMenuEntry(FText::FromString(*ProjectName), FText::FromString(ProjectPath), FSlateIcon(), CurrentProjectAction);
+			}
+		}
+
+		// add other top-level projects, ignoring the current project if we've shown it already
+		const TArray<FString>& AvailableGames = FGameProjectHelper::GetAvailableGames();
+		for (int32 Index = 0; Index < AvailableGames.Num(); ++Index)
+		{
+			FString ProjectPath = FPaths::RootDir() / AvailableGames[Index] / AvailableGames[Index] + TEXT(".uproject");
+			if (CurrentProjectOption == ECurrentProjectOption::ActualProject && !ProjectName.IsEmpty() && AvailableGames[Index] == ProjectName)
+			{
+				continue;
+			}
+
+			FUIAction ProjectAction(FExecuteAction::CreateSP(this, &SCustomLaunchProjectCombo::SetProjectPath, ProjectPath));
+			MenuBuilder.AddMenuEntry(FText::FromString(AvailableGames[Index]), FText::FromString(ProjectPath), FSlateIcon(), ProjectAction);
+		}
+
+		// browse option at the bottom
+		MenuBuilder.AddMenuSeparator();
+		FUIAction BrowseAction(FExecuteAction::CreateSP(this, &SCustomLaunchProjectCombo::OnBrowseForProject));
+		MenuBuilder.AddMenuEntry(LOCTEXT("BrowseAction", "Browse..."), LOCTEXT("BrowseActionHint", "Browse for a project on your computer"), FSlateIcon(), BrowseAction);
+	}
+
+	return MenuBuilder.MakeWidget();
+}
+END_SLATE_FUNCTION_BUILD_OPTIMIZATION
+
+
+void SCustomLaunchProjectCombo::OnBrowseForProject()
+{
+	FString DefaultPath = FPaths::RootDir();
+
+	const void* ParentWindowHandle = FSlateApplication::Get().FindBestParentWindowHandleForDialogs(nullptr);
+		
+	TArray<FString> OutFiles;
+	if (FDesktopPlatformModule::Get()->OpenFileDialog(ParentWindowHandle, LOCTEXT("SelectProjectDialogTitle", "Select a project").ToString(), DefaultPath, TEXT(""), TEXT("Project files (*.uproject)|*.uproject"), EFileDialogFlags::None, OutFiles))
+	{
+		SetProjectPath(OutFiles[0]);
+	}
+}
+
+FText SCustomLaunchProjectCombo::GetProjectName() const
+{
+	FString ProjectPath = SelectedProject.Get();
+
+	if (!ProjectPath.IsEmpty() && HasProject.Get())
+	{
+		return FText::FromString(FPaths::GetBaseFilename(ProjectPath));
+	}
+
+	if (!ProjectPath.IsEmpty() && CurrentProjectOption == ECurrentProjectOption::Empty)
+	{
+		return LOCTEXT("CurrentProjectAction", "Current Project");
+	}
+	else if (ProjectPath.IsEmpty() && bShowAnyProjectOption)
+	{
+		return LOCTEXT("AnyProjectAction", "Any Project");
+	}
+	else
+	{
+		return LOCTEXT("SelectProjectAction", "Select...");
+	}
+}
+
+void SCustomLaunchProjectCombo::SetProjectPath(FString ProjectPath)
+{
+	OnSelectionChanged.ExecuteIfBound(ProjectPath);
+}
+
+
+#undef LOCTEXT_NAMESPACE
diff --git a/Engine/Plugins/Developer/CustomLaunchUI/Source/Private/Widgets/Shared/SCustomLaunchProjectCombo.h b/Engine/Plugins/Developer/CustomLaunchUI/Source/Private/Widgets/Shared/SCustomLaunchProjectCombo.h
new file mode 100644
--- /dev/null
+++ b/Engine/Plugins/Developer/CustomLaunchUI/Source/Private/Widgets/Shared/SCustomLaunchProjectCombo.h
@@ -0,0 +1,57 @@
+// Copyright Epic Games, Inc. All Rights Reserved.
+
+#pragma once
+
+#include "CoreMinimal.h"
+#include "Widgets/SCompoundWidget.h"
+#include "Styling/SlateTypes.h"
+#include "Styling/CoreStyle.h"
+#include "Styling/SlateWidgetStyleAsset.h"
+#include "Fonts/SlateFontInfo.h"
+
+class SCustomLaunchProjectCombo
+	: public SCompoundWidget
+{
+public:
+	DECLARE_DELEGATE_OneParam(FOnSelectionChanged, FString );
+
+	enum class ECurrentProjectOption
+	{
+		None,
+		ActualProject,
+		Empty,
+	};
+
+
+	SLATE_BEGIN_ARGS(SCustomLaunchProjectCombo)
+		: _ShowAnyProjectOption(false)
+		, _CurrentProjectOption(ECurrentProjectOption::None)
+		, _TextStyle( &FCoreStyle::Get().GetWidgetStyle<FTextBlockStyle>( "NormalText" ) )
+		{}
+		SLATE_EVENT(FOnSelectionChanged, OnSelectionChanged);
+		SLATE_ATTRIBUTE(FString, SelectedProject);
+		SLATE_ATTRIBUTE(bool, HasProject);
+		SLATE_ARGUMENT(bool, ShowAnyProjectOption)
+		SLATE_ARGUMENT(ECurrentProjectOption, CurrentProjectOption)
+		SLATE_ATTRIBUTE(FSlateFontInfo, Font)
+		SLATE_STYLE_ARGUMENT( FTextBlockStyle, TextStyle )
+	SLATE_END_ARGS()
+
+public:
+	void Construct(	const FArguments& InArgs);
+
+protected:
+	TAttribute<FString> SelectedProject;
+	TAttribute<bool> HasProject;
+	FOnSelectionChanged OnSelectionChanged;
+	bool bShowAnyProjectOption;
+	ECurrentProjectOption CurrentProjectOption;
+
+	TSharedRef<SWidget> MakeProjectSelectionWidget();
+	void OnBrowseForProject();
+
+	FText GetProjectName() const;
+	void SetProjectPath(FString ProjectPath);
+
+
+};
diff --git a/Engine/Source/Editor/TurnkeySupport/Private/TurnkeySupportModule.cpp b/Engine/Source/Editor/TurnkeySupport/Private/TurnkeySupportModule.cpp
--- a/Engine/Source/Editor/TurnkeySupport/Private/TurnkeySupportModule.cpp
+++ b/Engine/Source/Editor/TurnkeySupport/Private/TurnkeySupportModule.cpp
@@ -33,6 +33,7 @@ IMPLEMENT_MODULE(FTurnkeySupportModuleEmpty, TurnkeySupport)
 #include "Interfaces/ITargetPlatform.h"
 #include "Interfaces/ITargetPlatformManagerModule.h"
 #include "Interfaces/IProjectManager.h"
+#include "Interfaces/IPluginManager.h"
 #include "SourceControlHelpers.h"
 #include "ISourceControlModule.h"
 #include "ITargetDeviceProxy.h"
@@ -57,6 +58,7 @@ IMPLEMENT_MODULE(FTurnkeySupportModuleEmpty, TurnkeySupport)
 #if WITH_EDITOR
 #include "LevelEditor.h"
 #include "Experimental/ZenServerInterface.h"
+#include "Settings/EditorExperimentalSettings.h"
 #endif
 
 #include "Misc/App.h"
@@ -231,6 +233,11 @@ public:
 		FGlobalTabmanager::Get()->TryInvokeTab(FTabId("ProjectLauncher"));
 	}
 
+	static void OpenCustomLaunchUI()
+	{
+		FGlobalTabmanager::Get()->TryInvokeTab(FTabId("CustomLaunchUI"));
+	}
+
 	static void OpenDeviceManager()
 	{
 		FGlobalTabmanager::Get()->TryInvokeTab(FTabId("DeviceManager"));
@@ -2252,6 +2259,27 @@ TSharedRef<SWidget> FTurnkeySupportModule::MakeTurnkeyMenuWidget() const
 		// options section
 		FToolMenuSection& OptionsSection = Menu->AddSection("TurnkeyOptions", LOCTEXT("TurnkeySection_Options", "Options and Settings"));
 		{
+#if WITH_EDITOR
+			if (IPluginManager::Get().FindEnabledPlugin(TEXT("CustomLaunchUI")) != nullptr)
+			{
+				OptionsSection.AddMenuEntry(
+					NAME_None,
+					LOCTEXT("OpenCustomLaunchProfileUI", "Custom Launch... (experimental)"),
+					LOCTEXT("OpenCustomLaunchProfileUI_ToolTip", "Configure custom launch profiles for advanced packaging, deploying and launching of your projects"),
+					FSlateIcon(FAppStyle::Get().GetStyleSetName(), "Launcher.TabIcon"),
+					FUIAction(
+						FExecuteAction::CreateStatic(&FTurnkeySupportCallbacks::OpenCustomLaunchUI), 
+						FCanExecuteAction(), 
+						FIsActionChecked(), 
+						FIsActionButtonVisible::CreateLambda([]() 
+						{
+							return GetDefault<UEditorExperimentalSettings>()->bShowNewLaunchProfileTools; 
+						}
+					)
+				));
+			}
+#endif
+
 			OptionsSection.AddMenuEntry(
 				NAME_None,
 				LOCTEXT("OpenProjectLauncher", "Project Launcher..."),
diff --git a/Engine/Source/Editor/UnrealEd/Classes/Settings/EditorExperimentalSettings.h b/Engine/Source/Editor/UnrealEd/Classes/Settings/EditorExperimentalSettings.h
--- a/Engine/Source/Editor/UnrealEd/Classes/Settings/EditorExperimentalSettings.h
+++ b/Engine/Source/Editor/UnrealEd/Classes/Settings/EditorExperimentalSettings.h
@@ -120,6 +120,10 @@ public:
 	UPROPERTY(EditAnywhere, config, Category = Tools, meta = (DisplayName = "Enable tools in Packaging Settings -> Additional builds for this project"))
 	bool bProjectCustomBuildTools;
 
+	/** Shows an option to configure custom launch profiles for advanced packaging, deploying and launching of your projects */
+	UPROPERTY(EditAnywhere, config, Category = Tools)
+	bool bShowNewLaunchProfileTools;
+
 	/** Enable experimental PIE preview device launch */
 	UPROPERTY(EditAnywhere, config, Category = PIE, meta = (DisplayName = "Enable mobile PIE with preview device launch options."))
 	bool bMobilePIEPreviewDeviceLaunch;
