Minor tweaks for custom launch UI widgets:
- Override font for platform combo box.
- Add check marks to build target selector menu.

#rnx

diff --git a/Engine/Plugins/Developer/CustomLaunchUI/Source/Private/Details/LauncherProfileSettingsDetails.cpp b/Engine/Plugins/Developer/CustomLaunchUI/Source/Private/Details/LauncherProfileSettingsDetails.cpp
--- a/Engine/Plugins/Developer/CustomLaunchUI/Source/Private/Details/LauncherProfileSettingsDetails.cpp
+++ b/Engine/Plugins/Developer/CustomLaunchUI/Source/Private/Details/LauncherProfileSettingsDetails.cpp
@@ -336,6 +336,7 @@ void FLauncherCustomProfileSettingsDetails::CustomizeDetails(IDetailLayoutBuilde
 		.SelectedPlatforms(this, &FLauncherCustomProfileSettingsDetails::GetSelectedPlatforms)
 		.OnSelectionChanged(this, &FLauncherCustomProfileSettingsDetails::SetSelectedPlatforms)
 		.BasicPlatformsOnly(CustomLaunchUI::bUseFriendlyBuildTargetSelection)
+		.Font(IDetailLayoutBuilder::GetDetailFont())
 	];
 
 
diff --git a/Engine/Plugins/Developer/CustomLaunchUI/Source/Private/Widgets/Shared/SCustomLaunchBuildTargetCombo.cpp b/Engine/Plugins/Developer/CustomLaunchUI/Source/Private/Widgets/Shared/SCustomLaunchBuildTargetCombo.cpp
--- a/Engine/Plugins/Developer/CustomLaunchUI/Source/Private/Widgets/Shared/SCustomLaunchBuildTargetCombo.cpp
+++ b/Engine/Plugins/Developer/CustomLaunchUI/Source/Private/Widgets/Shared/SCustomLaunchBuildTargetCombo.cpp
@@ -5,7 +5,6 @@
 #include "SlateOptMacros.h"
 #include "Widgets/Input/SComboButton.h"
 #include "Framework/MultiBox/MultiBoxBuilder.h"
-#include "Interfaces/IMainFrameModule.h"
 #include "DesktopPlatformModule.h"
 #include "Misc/Paths.h"
 #include "Settings/ProjectPackagingSettings.h"
@@ -44,11 +43,20 @@ TSharedRef<SWidget> SCustomLaunchBuildTargetCombo::MakeBuildTargetSelectionWidge
 {
 	FMenuBuilder MenuBuilder(true, nullptr);
 	{
-		FUIAction AnyProjectAction(FExecuteAction::CreateSP(this, &SCustomLaunchBuildTargetCombo::SetBuildTargetName, FString()));
-
 		FString DefaultBuildTarget = GetDefaultBuildTargetName();
 		FText DefaultBuildTargetName = DefaultBuildTarget.IsEmpty() ? LOCTEXT("DefaultBuildTargetName", "Project Default") : FText::Format(LOCTEXT("CurProjectDefaultBuildTargetName", "{0} (Project Default)"), FText::FromString(DefaultBuildTarget));
-		MenuBuilder.AddMenuEntry(DefaultBuildTargetName, LOCTEXT("DefaultBuildTargetActionHint", "Use the project default build target."), FSlateIcon(), AnyProjectAction);
+		MenuBuilder.AddMenuEntry(
+			DefaultBuildTargetName, 
+			LOCTEXT("DefaultBuildTargetActionHint", "Use the project default build target."), 
+			FSlateIcon(), 
+			FUIAction(
+				FExecuteAction::CreateSP(this, &SCustomLaunchBuildTargetCombo::SetBuildTargetName, FString()),
+				FCanExecuteAction(),
+				FGetActionCheckState::CreateLambda( [this]() { return SelectedBuildTarget.Get().IsEmpty() ? ECheckBoxState::Checked : ECheckBoxState::Unchecked; })
+			), 
+			NAME_None, 
+			EUserInterfaceActionType::Check);
+
 
 		if (SelectedProject.IsSet())
 		{
@@ -64,14 +72,29 @@ TSharedRef<SWidget> SCustomLaunchBuildTargetCombo::MakeBuildTargetSelectionWidge
 
 				if (!SupportedTargetTypes.IsSet() || SupportedTargetTypes.Get().IsEmpty() || SupportedTargetTypes.Get().Contains(BuildTarget.Type))
 				{
-					FUIAction ProjectAction(FExecuteAction::CreateSP(this, &SCustomLaunchBuildTargetCombo::SetBuildTargetName, BuildTarget.Name));
-					MenuBuilder.AddMenuEntry(FText::FromString(BuildTarget.Name), FText::FromString(BuildTarget.Path), FSlateIcon(), ProjectAction);
+					
+					MenuBuilder.AddMenuEntry(
+						FText::FromString(BuildTarget.Name), 
+						FText::FromString(BuildTarget.Path), 
+						FSlateIcon(),
+						FUIAction(
+							FExecuteAction::CreateSP(this, &SCustomLaunchBuildTargetCombo::SetBuildTargetName, BuildTarget.Name),
+							FCanExecuteAction(),
+							FGetActionCheckState::CreateLambda( [this, BuildTargetName = BuildTarget.Name]() { return (SelectedBuildTarget.Get() == BuildTargetName) ? ECheckBoxState::Checked : ECheckBoxState::Unchecked; })
+						),
+						NAME_None,
+						EUserInterfaceActionType::Check);
 				}
 				else
 				{
-					FUIAction ProjectAction(FExecuteAction::CreateSP(this, &SCustomLaunchBuildTargetCombo::SetBuildTargetName, BuildTarget.Name),
-											FCanExecuteAction::CreateLambda([]() { return false; }));
-					MenuBuilder.AddMenuEntry(FText::FromString(BuildTarget.Name), LOCTEXT("NotCompatible", "The selected platform does not support this build target type"), FSlateIcon(), ProjectAction);
+					MenuBuilder.AddMenuEntry(
+						FText::FromString(BuildTarget.Name), 
+						LOCTEXT("NotCompatible", "The selected platform does not support this build target type"), 
+						FSlateIcon(), 
+						FUIAction(
+							FExecuteAction::CreateSP(this, &SCustomLaunchBuildTargetCombo::SetBuildTargetName, BuildTarget.Name),
+							FCanExecuteAction::CreateLambda([]() { return false; })
+						));
 				}
 			}
 		}
diff --git a/Engine/Plugins/Developer/CustomLaunchUI/Source/Private/Widgets/Shared/SCustomLaunchPlatformCombo.cpp b/Engine/Plugins/Developer/CustomLaunchUI/Source/Private/Widgets/Shared/SCustomLaunchPlatformCombo.cpp
--- a/Engine/Plugins/Developer/CustomLaunchUI/Source/Private/Widgets/Shared/SCustomLaunchPlatformCombo.cpp
+++ b/Engine/Plugins/Developer/CustomLaunchUI/Source/Private/Widgets/Shared/SCustomLaunchPlatformCombo.cpp
@@ -20,6 +20,8 @@ void SCustomLaunchPlatformCombo::Construct(const FArguments& InArgs)
 	SelectedPlatforms = InArgs._SelectedPlatforms;
 	bBasicPlatformsOnly = InArgs._BasicPlatformsOnly;
 
+	FSlateFontInfo Font = InArgs._Font.IsSet() ? InArgs._Font.Get() : InArgs._TextStyle->Font;
+
 	ChildSlot
 	[
 		SAssignNew(PlatformsComboBox, SComboBox<TSharedPtr<FString>>)
@@ -46,6 +48,7 @@ void SCustomLaunchPlatformCombo::Construct(const FArguments& InArgs)
 			[
 				SNew(STextBlock)
 				.Text(this, &SCustomLaunchPlatformCombo::GetSelectedPlatformName)
+				.Font(Font)
 			]
 		]
 	];
diff --git a/Engine/Plugins/Developer/CustomLaunchUI/Source/Private/Widgets/Shared/SCustomLaunchPlatformCombo.h b/Engine/Plugins/Developer/CustomLaunchUI/Source/Private/Widgets/Shared/SCustomLaunchPlatformCombo.h
--- a/Engine/Plugins/Developer/CustomLaunchUI/Source/Private/Widgets/Shared/SCustomLaunchPlatformCombo.h
+++ b/Engine/Plugins/Developer/CustomLaunchUI/Source/Private/Widgets/Shared/SCustomLaunchPlatformCombo.h
@@ -4,6 +4,9 @@
 
 #include "CoreMinimal.h"
 #include "Widgets/SCompoundWidget.h"
+#include "Styling/CoreStyle.h"
+#include "Styling/SlateWidgetStyleAsset.h"
+#include "Fonts/SlateFontInfo.h"
 
 template<typename ItemType> class SComboBox;
 
@@ -14,10 +17,14 @@ class SCustomLaunchPlatformCombo
 public:
 	DECLARE_DELEGATE_OneParam(FOnSelectionChanged, TArray<FString> );
 
-	SLATE_BEGIN_ARGS(SCustomLaunchPlatformCombo){}
+	SLATE_BEGIN_ARGS(SCustomLaunchPlatformCombo)
+		: _TextStyle( &FCoreStyle::Get().GetWidgetStyle<FTextBlockStyle>( "NormalText" ) )
+		{}
 		SLATE_EVENT(FOnSelectionChanged, OnSelectionChanged);
 		SLATE_ATTRIBUTE(TArray<FString>, SelectedPlatforms);
 		SLATE_ARGUMENT(bool, BasicPlatformsOnly)
+		SLATE_ATTRIBUTE(FSlateFontInfo, Font)
+		SLATE_STYLE_ARGUMENT( FTextBlockStyle, TextStyle )
 	SLATE_END_ARGS()
 
 public:
