Style changes
- remove dependencies on editor-only styles
- adjust some background brush colours

#rnx

diff --git a/Engine/Plugins/Developer/CustomLaunchUI/Source/Private/Styling/CustomLaunchUIStyle.cpp b/Engine/Plugins/Developer/CustomLaunchUI/Source/Private/Styling/CustomLaunchUIStyle.cpp
--- a/Engine/Plugins/Developer/CustomLaunchUI/Source/Private/Styling/CustomLaunchUIStyle.cpp
+++ b/Engine/Plugins/Developer/CustomLaunchUI/Source/Private/Styling/CustomLaunchUIStyle.cpp
@@ -20,10 +20,13 @@
 
 TSharedPtr<FSlateStyleSet> FCustomLaunchUIStyle::StyleSet = nullptr;
 
-static const FVector2D Icon15x15(16.0f, 16.0f);
+static const FVector2D Icon12x12(12.0f, 12.0f);
+static const FVector2D Icon15x15(15.0f, 15.0f);
 static const FVector2D Icon16x16(16.0f, 16.0f);
+static const FVector2D Icon24x24(24.0f, 24.0f);
 static const FVector2D Icon28x28(28.0f, 28.0f);
 static const FVector2D Icon36x36(36.0f, 36.0f);
+static const FVector2D Icon64x64(64.0f, 64.0f);
 static const FVector2D Icon512x512(512.0f, 512.0f);
 
 
@@ -35,6 +38,8 @@ void FCustomLaunchUIStyle::Initialize()
 	}
 
 	StyleSet = MakeShareable(new FSlateStyleSet("CustomLaunchUIStyle") );
+
+	StyleSet->SetParentStyleName("CoreStyle");
 	StyleSet->SetContentRoot(IPluginManager::Get().FindPlugin("CustomLaunchUI")->GetBaseDir() / TEXT("Resources"));
 	StyleSet->SetCoreContentRoot(FPaths::EngineContentDir() / TEXT("Slate"));
 
@@ -58,7 +63,32 @@ void FCustomLaunchUIStyle::Initialize()
 	StyleSet->Set("State.Canceled", FLinearColor(FColor(0xFF383838)) );
 	StyleSet->Set("State.AllComplete", FLinearColor(FColor(0xFF0070E0)) );
 
-	
+	FButtonStyle HoverHintOnly = FButtonStyle()
+			.SetNormal( FSlateNoResource() )
+			.SetHovered( CORE_BOX_BRUSH( "Common/ButtonHoverHint", FMargin(4.f/16.0f), FLinearColor(1.f,1.f,1.f,0.15f) ) )
+			.SetPressed( CORE_BOX_BRUSH( "Common/ButtonHoverHint", FMargin(4.f/16.0f), FLinearColor(1.f,1.f,1.f,0.25f) ) )
+			.SetNormalPadding( FMargin(0.f,0.f,0.f,1.f) )
+			.SetPressedPadding( FMargin(0.f,1.f,0.f,0.f) );
+	StyleSet->Set( "HoverHintOnly", HoverHintOnly );
+
+	StyleSet->Set("Icons.DiffersFromDefault", new CORE_IMAGE_BRUSH_SVG( "../Editor/Slate/Starship/Common/ResetToDefault", Icon16x16) );
+	StyleSet->Set("Icons.DeviceManager", new CORE_IMAGE_BRUSH_SVG( "../Editor/Slate/Starship/Common/DeviceManager", Icon16x16) );
+	StyleSet->Set("Icons.Asset", new CORE_IMAGE_BRUSH( "Icons/doc_16x", Icon16x16 ) );
+
+	StyleSet->Set( "PathPickerButton", new CORE_IMAGE_BRUSH( "Icons/ellipsis_12x", Icon12x12 ) );
+
+	StyleSet->Set("Icons.Task.Run",	new CORE_IMAGE_BRUSH_SVG("../Editor/Slate/Starship/Common/Timecode", Icon16x16));
+	StyleSet->Set("Icons.Task.Cleanup", new CORE_IMAGE_BRUSH_SVG("../Editor/Slate/Starship/GraphEditors/CleanUp", Icon16x16));
+	StyleSet->Set("Icons.Task.Launch",	new CORE_IMAGE_BRUSH_SVG("Starship/Common/ProjectLauncher", Icon16x16));
+	StyleSet->Set("Icons.Task.Build", new CORE_IMAGE_BRUSH_SVG("../Editor/Slate/Starship/Common/Adjust", Icon16x16));
+	StyleSet->Set("Icons.Task.Cook", new CORE_IMAGE_BRUSH_SVG( "../Editor/Slate/Starship/Common/CookContent", Icon16x16 ) );
+	StyleSet->Set("Icons.Task.Deploy", new CORE_IMAGE_BRUSH_SVG( "../Editor/Slate/Starship/Common/DeviceManager", Icon16x16 ) );
+	StyleSet->Set("Icons.Task.Archive", new CORE_IMAGE_BRUSH_SVG("../Editor/Slate/Starship/Common/SaveCurrent", Icon16x16));
+	StyleSet->Set("Icons.Task.Package", new CORE_IMAGE_BRUSH_SVG("../Editor/Slate/Starship/Common/ProjectPackage", Icon16x16));
+
+	StyleSet->Set("Profile.NoPlatform", new CORE_IMAGE_BRUSH_SVG("Starship/Launcher/PaperAirplane", Icon24x24));
+	StyleSet->Set("Profile.NoPlatform.Large", new CORE_IMAGE_BRUSH_SVG("Starship/Launcher/PaperAirplane", Icon64x64));
+
 	FSlateStyleRegistry::RegisterSlateStyle( *StyleSet.Get() );
 }
 
@@ -95,40 +125,39 @@ const FSlateBrush* FCustomLaunchUIStyle::GetBrushForTask(ILauncherTaskPtr Task)
 		FString TaskName = Task->GetName();
 		if (TaskName.Contains(TEXT("Cooking in the editor")))
 		{
-			return FCoreStyle::Get().GetBrush("GraphEditor.Timeline_16x");
+			return GetBrush("Icons.Task.Run");
 		}
 		else if (TaskName.Contains(TEXT("Post Launch")))
 		{
-			return FCoreStyle::Get().GetBrush("GraphEditor.Clean");
+			return GetBrush("Icons.Task.Cleanup");
 		}
 		else if (TaskName.Contains(TEXT("Launch")))
 		{
-			return FCoreStyle::Get().GetBrush("Icons.Launch");
+			return GetBrush("Icons.Task.Launch");
 		}
 		else if (TaskName.Contains(TEXT("Build")))
 		{
-			return FCoreStyle::Get().GetBrush("Icons.Adjust");
+			return GetBrush("Icons.Task.Build");
 		}
 		else if (TaskName.Contains(TEXT("Cook")))
 		{
-			return FCoreStyle::Get().GetBrush("MainFrame.CookContent");
+			return GetBrush("Icons.Task.Cook");
 		}
 		else if (TaskName.Contains(TEXT("Deploy")))
 		{
-			return FCoreStyle::Get().GetBrush("DeviceDetails.TabIcon");
+			return GetBrush("Icons.Task.Deploy");
 		}
 		else if (TaskName.Contains(TEXT("Archive")))
 		{
-			return FCoreStyle::Get().GetBrush("Icons.Save");
+			return GetBrush("Icons.Task.Archive");
 		}
 		else if (TaskName.Contains(TEXT("Run")))
 		{
-			return FCoreStyle::Get().GetBrush("GraphEditor.Timeline_16x");
+			return GetBrush("Icons.Task.Run");
 		}
-
 	}
 
-	return FCoreStyle::Get().GetBrush("Icons.Package");
+	return GetBrush("Icons.Task.Package");
 }
 
 
@@ -138,6 +167,13 @@ const FSlateBrush* FCustomLaunchUIStyle::GetProfileBrushForPlatform(const Platfo
 	{
 		return FAppStyle::GetBrush(PlatformInfo->GetIconStyleName(IconSize));
 	}
-	
-	return FAppStyle::GetBrush("LauncherCommand.QuickLaunch");
+
+	if (IconSize == EPlatformIconSize::Normal)
+	{	
+		return GetBrush("Profile.NoPlatform");
+	}
+	else
+	{
+		return GetBrush("Profile.NoPlatform.Large");
+	}
 }
diff --git a/Engine/Plugins/Developer/CustomLaunchUI/Source/Private/Widgets/Profiles/SCustomLaunchCustomProfileSelector.cpp b/Engine/Plugins/Developer/CustomLaunchUI/Source/Private/Widgets/Profiles/SCustomLaunchCustomProfileSelector.cpp
--- a/Engine/Plugins/Developer/CustomLaunchUI/Source/Private/Widgets/Profiles/SCustomLaunchCustomProfileSelector.cpp
+++ b/Engine/Plugins/Developer/CustomLaunchUI/Source/Private/Widgets/Profiles/SCustomLaunchCustomProfileSelector.cpp
@@ -193,11 +193,11 @@ TSharedRef<ITableRow> SCustomLaunchCustomProfileSelector::GenerateCustomProfileR
 			[
 				SNew(SButton)
 				.OnClicked(this, &SCustomLaunchCustomProfileSelector::OnOpenDeviceManagerClicked)
-				.ButtonStyle(FAppStyle::Get(), "HoverHintOnly")
+				.ButtonStyle(FCustomLaunchUIStyle::Get(), "HoverHintOnly")
 				.ToolTipText(LOCTEXT("OpenDeviceManagerToolTip", "Open the Device Manager window, where you can setup and claim devices connected to your machine or shared on the network."))
 				[
 					SNew(SImage)
-					.Image(FAppStyle::Get().GetBrush("DeviceDetails.TabIcon"))
+					.Image(FCustomLaunchUIStyle::Get().GetBrush("Icons.DeviceManager"))
 					.ColorAndOpacity(FColor::White)
 				]
 			]
@@ -246,7 +246,7 @@ TSharedRef<ITableRow> SCustomLaunchCustomProfileSelector::GenerateCustomProfileR
 	[
 		SNew(SButton)
 		.OnClicked(this, &SCustomLaunchCustomProfileSelector::OnEditProfileClicked, Profile)
-		.ButtonStyle(FAppStyle::Get(), "HoverHintOnly")
+		.ButtonStyle(FCustomLaunchUIStyle::Get(), "HoverHintOnly")
 		.Visibility( this, &SCustomLaunchCustomProfileSelector::GetInlineEditButtonVisibility, RowContent->AsWidget(), Profile )
 		[
 			SNew(SImage)
@@ -263,7 +263,7 @@ TSharedRef<ITableRow> SCustomLaunchCustomProfileSelector::GenerateCustomProfileR
 	[
 		SNew(SButton)
 		.OnClicked(this, &SCustomLaunchCustomProfileSelector::OnCloseEditorClicked)
-		.ButtonStyle(FAppStyle::Get(), "HoverHintOnly")
+		.ButtonStyle(FCustomLaunchUIStyle::Get(), "HoverHintOnly")
 		.Visibility( this, &SCustomLaunchCustomProfileSelector::GetCloseEditorButtonVisibility, Profile)
 		[
 			SNew(SImage)
diff --git a/Engine/Plugins/Developer/CustomLaunchUI/Source/Private/Widgets/Profiles/SCustomLaunchProfilesPanel.cpp b/Engine/Plugins/Developer/CustomLaunchUI/Source/Private/Widgets/Profiles/SCustomLaunchProfilesPanel.cpp
--- a/Engine/Plugins/Developer/CustomLaunchUI/Source/Private/Widgets/Profiles/SCustomLaunchProfilesPanel.cpp
+++ b/Engine/Plugins/Developer/CustomLaunchUI/Source/Private/Widgets/Profiles/SCustomLaunchProfilesPanel.cpp
@@ -202,21 +202,26 @@ TSharedRef<SWidget> SCustomLaunchProfilesPanel::CreateProfilesPanel()
 	.FillHeight(1)
 	.Padding(8,8,8,0)
 	[
-		SAssignNew(ProfileEditorSplitter, SSplitter)
-		.Orientation(Orient_Horizontal)
-		.PhysicalSplitterHandleSize(8)
-		.Style(FAppStyle::Get(), "SplitterPanel")
-			
-		// current profile selection panel
-		+SSplitter::Slot()
+		SNew(SBorder)
+		.BorderImage(FAppStyle::GetBrush("ChildWindow.Background"))
+		.Padding(0)
 		[
-			CreateProfileSelectorWidget()
-		]
+			SAssignNew(ProfileEditorSplitter, SSplitter)
+			.Orientation(Orient_Horizontal)
+			.PhysicalSplitterHandleSize(8)
+			.Style(FAppStyle::Get(), "SplitterPanel")
+			
+			// current profile selection panel
+			+SSplitter::Slot()
+			[
+				CreateProfileSelectorWidget()
+			]
 		
-		// Profile editor panel
-		+SSplitter::Slot()
-		[
-			CreateProfileEditorWidget()
+			// Profile editor panel
+			+SSplitter::Slot()
+			[
+				CreateProfileEditorWidget()
+			]
 		]
 	];
 
@@ -513,7 +518,7 @@ TSharedRef<SWidget> SCustomLaunchProfilesPanel::CreateAdvancedProfileWarningWidg
 			.Padding(2)
 			[
 				SNew(STextBlock)
-				.TextStyle(FAppStyle::Get(), "LargeText")
+				.ColorAndOpacity(FLinearColor::White)
 				.Text(LOCTEXT("AdvancedProfileWarning", "Profile contains advanced/legacy settings"))
 			]
 
@@ -522,8 +527,8 @@ TSharedRef<SWidget> SCustomLaunchProfilesPanel::CreateAdvancedProfileWarningWidg
 			.Padding(2)
 			[
 				SNew(STextBlock)
-				.TextStyle(FAppStyle::Get(), "NormalText.Important")
-				.Text(LOCTEXT("AdvancedProfileWarningDetail", "These settings may be multiple cook platforms or cultures, or distribution packages, DLC or patch builds. Please use Project Launcher to edit the advanced properties in this profile."))
+				.ColorAndOpacity(FLinearColor::White)
+				.Text(LOCTEXT("AdvancedProfileWarningDetail", "These settings may be multiple cook platforms or cultures, or distribution packages, DLC or patch builds.\nPlease use Project Launcher to edit the advanced properties in this profile."))
 			]
 		]
 
@@ -534,7 +539,7 @@ TSharedRef<SWidget> SCustomLaunchProfilesPanel::CreateAdvancedProfileWarningWidg
 		[
 			SNew(SButton)
 			.OnClicked_Lambda( []() { FGlobalTabmanager::Get()->TryInvokeTab(FTabId("ProjectLauncher")); return FReply::Handled(); })
-			.ButtonStyle(FAppStyle::Get(), "HoverHintOnly")
+			.ButtonStyle(FCustomLaunchUIStyle::Get(), "HoverHintOnly")
 			.ToolTipText(LOCTEXT("OpenDeviceProjectLauncherTip", "Open the legacy Project Launcher."))
 			[
 				SNew(SImage)
diff --git a/Engine/Plugins/Developer/CustomLaunchUI/Source/Private/Widgets/Shared/SCustomLaunchDeviceListView.cpp b/Engine/Plugins/Developer/CustomLaunchUI/Source/Private/Widgets/Shared/SCustomLaunchDeviceListView.cpp
--- a/Engine/Plugins/Developer/CustomLaunchUI/Source/Private/Widgets/Shared/SCustomLaunchDeviceListView.cpp
+++ b/Engine/Plugins/Developer/CustomLaunchUI/Source/Private/Widgets/Shared/SCustomLaunchDeviceListView.cpp
@@ -31,7 +31,7 @@ void SCustomLaunchDeviceListView::Construct(const FArguments& InArgs)
 	[
 		SNew(SBorder)
 		.Padding(2)
-		.BorderImage(FAppStyle::GetBrush("ToolPanel.GroupBorder"))
+		.BorderImage(FAppStyle::GetBrush("Brushes.Background"))
 		[
 			SAssignNew(DeviceProxyListView, SListView<TSharedPtr<ITargetDeviceProxy>>)
 			.ListItemsSource(&DeviceProxyList)
diff --git a/Engine/Plugins/Developer/CustomLaunchUI/Source/Private/Widgets/Shared/SCustomLaunchMapListView.cpp b/Engine/Plugins/Developer/CustomLaunchUI/Source/Private/Widgets/Shared/SCustomLaunchMapListView.cpp
--- a/Engine/Plugins/Developer/CustomLaunchUI/Source/Private/Widgets/Shared/SCustomLaunchMapListView.cpp
+++ b/Engine/Plugins/Developer/CustomLaunchUI/Source/Private/Widgets/Shared/SCustomLaunchMapListView.cpp
@@ -12,6 +12,7 @@
 #include "Widgets/Views/STableRow.h"
 #include "SlateOptMacros.h"
 #include "Styling/AppStyle.h"
+#include "Styling/CustomLaunchUIStyle.h"
 #include "Misc/PackageName.h"
 #include "HAL/FileManager.h"
 
@@ -28,10 +29,15 @@ void SCustomLaunchMapListView::Construct(const FArguments& InArgs)
 
 	ChildSlot
 	[
-		SAssignNew(MapTreeView, STreeView<FMapTreeNodePtr>)
-		.TreeItemsSource(&MapTreeViewItemsSource)
-		.OnGenerateRow(this, &SCustomLaunchMapListView::GenerateMapTreeNodeRow)
-		.OnGetChildren(this, &SCustomLaunchMapListView::GetMapTreeNodeChildren)
+		SNew(SBorder)
+		.Padding(2)
+		.BorderImage(FAppStyle::GetBrush("Brushes.Background"))
+		[
+			SAssignNew(MapTreeView, STreeView<FMapTreeNodePtr>)
+			.TreeItemsSource(&MapTreeViewItemsSource)
+			.OnGenerateRow(this, &SCustomLaunchMapListView::GenerateMapTreeNodeRow)
+			.OnGetChildren(this, &SCustomLaunchMapListView::GetMapTreeNodeChildren)
+		]
 	];
 
 	RefreshMapList();
@@ -127,11 +133,7 @@ const FSlateBrush* SCustomLaunchMapListView::GetMapTreeNodeIcon(FMapTreeNodePtr
 {
 	if (Node->Children.Num() == 0)
 	{
-#if WITH_EDITOR
-		return FAppStyle::GetBrush("ContentBrowser.ColumnViewAssetIcon");
-#else
-		return FAppStyle::GetBrush("Icons.FilledCircle"); // nice asset icon not available outside of the editor... @todo: find a better one than this!
-#endif
+		return FCustomLaunchUIStyle::GetBrush("Icons.Asset");
 	}
 	else if (MapTreeView->IsItemExpanded(Node))
 	{
