localization build fix

#rnx

diff --git a/Engine/Plugins/Developer/CustomLaunchUI/Source/Private/Widgets/Shared/SCustomLaunchDeviceWidgetBase.cpp b/Engine/Plugins/Developer/CustomLaunchUI/Source/Private/Widgets/Shared/SCustomLaunchDeviceWidgetBase.cpp
--- a/Engine/Plugins/Developer/CustomLaunchUI/Source/Private/Widgets/Shared/SCustomLaunchDeviceWidgetBase.cpp
+++ b/Engine/Plugins/Developer/CustomLaunchUI/Source/Private/Widgets/Shared/SCustomLaunchDeviceWidgetBase.cpp
@@ -125,4 +125,3 @@ const TSharedRef<ITargetDeviceProxyManager> SCustomLaunchDeviceWidgetBase::GetDe
 }
 
 
-#undef LOCTEXT_NAMESPACE
