Replace the PropertyEditor module dependency with a bespoke custom property editor.
- necessary to be able to add CustomLaunchUI to UnrealFrontend.
- some minor tweaks to non-editor logic in preparation.

#rnx

diff --git a/Engine/Plugins/Developer/CustomLaunchUI/Source/CustomLaunchUI.Build.cs b/Engine/Plugins/Developer/CustomLaunchUI/Source/CustomLaunchUI.Build.cs
--- a/Engine/Plugins/Developer/CustomLaunchUI/Source/CustomLaunchUI.Build.cs
+++ b/Engine/Plugins/Developer/CustomLaunchUI/Source/CustomLaunchUI.Build.cs
@@ -20,7 +20,6 @@ public class CustomLaunchUI : ModuleRules
 			"WorkspaceMenuStructure",
 			"ToolWidgets",
 			"TargetDeviceServices",
-			"PropertyEditor",
 			"Projects",
 			"DeveloperToolSettings",
 		]);
diff --git a/Engine/Plugins/Developer/CustomLaunchUI/Source/Private/Classes/LauncherProfileSettings.cpp b/Engine/Plugins/Developer/CustomLaunchUI/Source/Private/Classes/LauncherProfileSettings.cpp
deleted file mode 100644
diff --git a/Engine/Plugins/Developer/CustomLaunchUI/Source/Private/Classes/LauncherProfileSettings.h b/Engine/Plugins/Developer/CustomLaunchUI/Source/Private/Classes/LauncherProfileSettings.h
deleted file mode 100644
diff --git a/Engine/Plugins/Developer/CustomLaunchUI/Source/Private/CustomLaunchUIModule.cpp b/Engine/Plugins/Developer/CustomLaunchUI/Source/Private/CustomLaunchUIModule.cpp
--- a/Engine/Plugins/Developer/CustomLaunchUI/Source/Private/CustomLaunchUIModule.cpp
+++ b/Engine/Plugins/Developer/CustomLaunchUI/Source/Private/CustomLaunchUIModule.cpp
@@ -7,12 +7,12 @@
 #include "Widgets/Docking/SDockTab.h"
 #include "Widgets/SCustomLaunchUI.h"
 #include "Styling/CustomLaunchUIStyle.h"
-#include "Details/LauncherProfileSettingsDetails.h"
 #include "ILauncherServicesModule.h"
 #include "ITargetDeviceServicesModule.h"
-#include "PropertyEditorModule.h"
+#include "Marshalling/ILaunchProfileMarshaller.h"
+#include "Marshalling/BasicProfileMarshaller.h"
+#include "Marshalling/CustomProfileMarshaller.h"
 #include "Model/CustomLaunchUIModel.h"
-#include "Classes/LauncherProfileSettings.h"
 
 #define LOCTEXT_NAMESPACE "FCustomLaunchUIModule"
 
@@ -24,18 +24,9 @@ public:
 
 	virtual void StartupModule() override
 	{
-		// register detail customization
-		static FName PropertyEditor("PropertyEditor");
-		FPropertyEditorModule* PropertyModule = FModuleManager::GetModulePtr<FPropertyEditorModule>(PropertyEditor);
-
-		if (PropertyModule != nullptr)
-		{
-			static FName LauncherBasicLaunchProfileSettings("LauncherBasicLaunchProfileSettings");
-			PropertyModule->RegisterCustomClassLayout(LauncherBasicLaunchProfileSettings, FOnGetDetailCustomizationInstance::CreateStatic(&FLauncherBasicLaunchProfileSettingsDetails::MakeInstance));
-
-			static FName LauncherCustomProfileSettings("LauncherCustomProfileSettings");
-			PropertyModule->RegisterCustomClassLayout(LauncherCustomProfileSettings, FOnGetDetailCustomizationInstance::CreateStatic(&FLauncherCustomProfileSettingsDetails::MakeInstance));
-		}
+		// register marshalls
+		ILaunchProfileMarshallerFactory::Register( MakeShared<FBasicProfileMarshallerFactory>() );
+		ILaunchProfileMarshallerFactory::Register( MakeShared<FCustomProfileMarshallerFactory>() );
 
 		// register custom styles
 		FCustomLaunchUIStyle::Initialize();
@@ -64,10 +55,6 @@ private:
 
 		TSharedRef<CustomLaunchUI::FModel> Model = MakeShared<CustomLaunchUI::FModel>(TargetDeviceServicesModule.GetDeviceProxyManager(), LauncherServicesModule.CreateLauncher(), LauncherServicesModule.GetProfileManager() );
 		
-		// configure CDOs, for 'reset to default'
-		ULauncherBasicLaunchProfileSettings::StaticClass()->GetDefaultObject<ULauncherBasicLaunchProfileSettings>()->SetProfile(Model, Model->CreateBasicLaunchProfile());
-		ULauncherCustomProfileSettings::StaticClass()->     GetDefaultObject<ULauncherCustomProfileSettings>()->     SetProfile(Model, Model->CreateCustomProfile( TEXT("DefaultCustomProfile") ));
-
 		return SNew(SDockTab)
 		.TabRole(ETabRole::NomadTab)
 		[
diff --git a/Engine/Plugins/Developer/CustomLaunchUI/Source/Private/Details/LauncherProfileSettingsDetails.cpp b/Engine/Plugins/Developer/CustomLaunchUI/Source/Private/Details/LauncherProfileSettingsDetails.cpp
deleted file mode 100644
diff --git a/Engine/Plugins/Developer/CustomLaunchUI/Source/Private/Details/LauncherProfileSettingsDetails.h b/Engine/Plugins/Developer/CustomLaunchUI/Source/Private/Details/LauncherProfileSettingsDetails.h
deleted file mode 100644
diff --git a/Engine/Plugins/Developer/CustomLaunchUI/Source/Private/Marshalling/BasicProfileMarshaller.cpp b/Engine/Plugins/Developer/CustomLaunchUI/Source/Private/Marshalling/BasicProfileMarshaller.cpp
new file mode 100644
--- /dev/null
+++ b/Engine/Plugins/Developer/CustomLaunchUI/Source/Private/Marshalling/BasicProfileMarshaller.cpp
@@ -0,0 +1,275 @@
+// Copyright Epic Games, Inc. All Rights Reserved.
+
+#include "Marshalling/BasicProfileMarshaller.h"
+#include "Model/CustomLaunchUIModel.h"
+#include "Widgets/Shared/SCustomLaunchProjectCombo.h"
+#include "Widgets/Shared/SCustomLaunchBuildTargetCombo.h"
+#include "Widgets/Shared/SCustomLaunchBuildConfigurationCombo.h"
+#include "Widgets/Shared/SCustomLaunchContentSchemeCombo.h"
+#include "Widgets/Shared/SCustomLaunchDeviceListView.h"
+#include "Widgets/Input/SMultiLineEditableTextBox.h"
+#include "SlateOptMacros.h"
+
+
+#define LOCTEXT_NAMESPACE "CustomProfileMarshaller"
+
+class FBasicProfileMarshaller : public ILaunchProfileMarshaller, public TSharedFromThis<FBasicProfileMarshaller>
+{
+public:
+	FBasicProfileMarshaller( const ILauncherProfileRef& Profile, const TSharedRef<CustomLaunchUI::FModel>& InModel );
+	virtual ~FBasicProfileMarshaller() = default;
+
+	virtual void Construct() override;
+
+	virtual FString GetName() const override
+	{
+		return TEXT("BasicProfile");
+	}
+
+	virtual FLaunchProfileTreeDataRef GetProfileTree() override
+	{
+		return TreeData;
+	}
+
+protected:
+
+	FString GetProjectPath() const;
+	void SetProjectName(FString ProjectPath);
+	bool HasProject() const;
+
+	void SetBuildTarget(FString BuildTarget);
+	FString GetBuildTarget() const;
+
+	void SetBuildConfiguration(EBuildConfiguration BuildConfiguration);
+	EBuildConfiguration GetBuildConfiguration() const;
+
+	void SetContentScheme(CustomLaunchUI::EContentScheme ContentScheme);
+	bool IsContentSchemeAvailable(CustomLaunchUI::EContentScheme) const;
+
+	FText GetCommandLine() const;
+	void SetCommandLine( const FText& NewText, ETextCommit::Type InTextCommit );
+
+	void SetDeployDeviceIDs(TArray<FString> DeployDeviceIDs);
+	TArray<FString> GetDeployDeviceIDs() const;
+
+protected:
+	virtual void OnPropertyChanged() override;
+
+	FLaunchProfileTreeDataRef TreeData;
+	const ILauncherProfileRef Profile;
+	const ILauncherProfileRef DefaultProfile;
+	const TSharedRef<CustomLaunchUI::FModel> Model;
+	CustomLaunchUI::EContentScheme ContentScheme;
+};
+
+
+
+FBasicProfileMarshaller::FBasicProfileMarshaller( const ILauncherProfileRef& InProfile, const TSharedRef<CustomLaunchUI::FModel>& InModel )
+	: TreeData( MakeShared<FLaunchProfileTreeData>(InProfile, InModel, this) )
+	, Profile(InProfile)
+	, DefaultProfile(InModel->GetDefaultBasicLaunchProfile())
+	, Model(InModel)
+{
+}
+
+BEGIN_SLATE_FUNCTION_BUILD_OPTIMIZATION
+void FBasicProfileMarshaller::Construct()
+{
+	TreeData->AddHeading( LOCTEXT("GeneralSettingsHeading", "General Settings") )
+
+		.AddWidget( LOCTEXT("ProjectLabel", "Project"), 
+			SNew(SCustomLaunchProjectCombo)
+			.OnSelectionChanged(this, &FBasicProfileMarshaller::SetProjectName)
+			.SelectedProject(this, &FBasicProfileMarshaller::GetProjectPath)
+			.HasProject(this, &FBasicProfileMarshaller::HasProject)
+			.CurrentProjectOption(SCustomLaunchProjectCombo::ECurrentProjectOption::Empty)
+			.Font(FCoreStyle::Get().GetFontStyle("SmallFont"))
+		)
+
+		.AddWidget( LOCTEXT("TargetLabel", "Target"), 
+			{
+				.IsDefault = [this]()		{ return !Profile->HasBuildTargetSpecified() || Profile->GetBuildTarget().IsEmpty(); },
+				.SetToDefault = [this]()	{ SetBuildTarget(FString()); },
+			},
+			SNew(SCustomLaunchBuildTargetCombo)
+			.OnSelectionChanged(this, &FBasicProfileMarshaller::SetBuildTarget)
+			.SelectedBuildTarget(this, &FBasicProfileMarshaller::GetBuildTarget)
+			.SelectedProject(this, &FBasicProfileMarshaller::GetProjectPath)
+			.Font(FCoreStyle::Get().GetFontStyle("SmallFont"))
+		)
+
+		.AddWidget( LOCTEXT("ConfigurationLabel", "Configuration"), 
+			SNew(SCustomLaunchBuildConfigurationCombo)
+			.OnSelectionChanged(this, &FBasicProfileMarshaller::SetBuildConfiguration)
+			.SelectedBuildConfiguration(this, &FBasicProfileMarshaller::GetBuildConfiguration)
+			.Font(FCoreStyle::Get().GetFontStyle("SmallFont"))
+		)
+
+		.AddWidget( LOCTEXT("ContentSchemeLabel", "Content Scheme"), 
+			SNew(SCustomLaunchContentSchemeCombo)
+			.OnSelectionChanged(this, &FBasicProfileMarshaller::SetContentScheme)
+			.SelectedContentScheme_Lambda( [this]() { return ContentScheme; } )
+			.IsContentSchemeAvailable(this, &FBasicProfileMarshaller::IsContentSchemeAvailable)
+			.Font(FCoreStyle::Get().GetFontStyle("SmallFont"))
+		)
+
+		.AddWidget( LOCTEXT("CommandLineLabel", "Command Line"), 
+			{
+				.IsDefault = [this]()		{ return GetCommandLine().IsEmpty(); },
+				.SetToDefault = [this]()	{ SetCommandLine(FText(), ETextCommit::Default); },
+			},
+			SNew(SMultiLineEditableTextBox)
+			.AllowMultiLine(false)
+			.AutoWrapText(true)
+			.Text(this, &FBasicProfileMarshaller::GetCommandLine)
+			.OnTextCommitted(this, &FBasicProfileMarshaller::SetCommandLine)
+		)
+
+		.AddWidget( LOCTEXT("TargetDeviceLabel", "Target Device"), 
+			SNew(SCustomLaunchDeviceListView)
+			.OnSelectionChanged(this, &FBasicProfileMarshaller::SetDeployDeviceIDs)
+			.SelectedDevices(this, &FBasicProfileMarshaller::GetDeployDeviceIDs)
+			.AllPlatforms(true)
+			.SingleSelect(true)
+		)
+	;
+
+	ContentScheme = Model->DetermineProfileContentScheme(Profile);
+}
+END_SLATE_FUNCTION_BUILD_OPTIMIZATION
+
+FString FBasicProfileMarshaller::GetProjectPath() const
+{
+	if (Profile->HasProjectSpecified())
+	{
+		return Profile->GetProjectPath();
+	}
+	else
+	{
+		return Model->GetProfileManager()->GetProjectPath();
+	}
+}
+
+void FBasicProfileMarshaller::SetProjectName(FString ProjectPath)
+{
+	Profile->SetProjectPath(ProjectPath);
+	Profile->SetProjectSpecified(!ProjectPath.IsEmpty());
+
+	OnPropertyChanged();
+}
+
+bool FBasicProfileMarshaller::HasProject() const
+{
+	return Profile->HasProjectSpecified();
+}
+
+FString FBasicProfileMarshaller::GetBuildTarget() const
+{
+	if (Profile->HasBuildTargetSpecified())
+	{
+		return Profile->GetBuildTarget();
+	}
+	else
+	{
+		return Model->GetProfileManager()->GetBuildTarget();
+	}
+}
+
+void FBasicProfileMarshaller::SetBuildTarget(FString BuildTarget)
+{
+	Profile->SetBuildTarget(BuildTarget);
+	Profile->SetBuildTargetSpecified(!BuildTarget.IsEmpty());
+
+	Model->UpdatedCookedPlatformsFromDeployDeviceProxy(Profile);
+
+	OnPropertyChanged();
+}
+
+void FBasicProfileMarshaller::SetBuildConfiguration(EBuildConfiguration BuildConfiguration)
+{
+	Profile->SetBuildConfiguration(BuildConfiguration);
+	OnPropertyChanged();
+}
+
+EBuildConfiguration FBasicProfileMarshaller::GetBuildConfiguration() const
+{
+	return Profile->GetBuildConfiguration();
+}
+
+void FBasicProfileMarshaller::SetContentScheme(CustomLaunchUI::EContentScheme InContentScheme)
+{
+	Model->SetProfileContentScheme(InContentScheme, Profile);
+	OnPropertyChanged();
+
+	// refresh the cached content scheme in case the option that was selected is not available
+	ContentScheme = Model->DetermineProfileContentScheme(Profile);
+}
+
+bool FBasicProfileMarshaller::IsContentSchemeAvailable(CustomLaunchUI::EContentScheme InContentScheme) const
+{
+	// @fixme: check the project supports the content scheme etc
+	return InContentScheme != CustomLaunchUI::EContentScheme::ZenPakStreaming;
+}
+
+FText FBasicProfileMarshaller::GetCommandLine() const
+{
+	// get unified command line from these two fields. the first is presented in "Build" for old Project Launcher, and the latter is presented in "Launch" for old Project Launcher)
+	// when we save back to the profile, this will be stored just in the "Build" one for clarity (because multiple roles are not supported in old or new project launcher)
+	FString CommandLine = Profile->GetAdditionalCommandLineParameters().TrimStartAndEnd() + TEXT(" ") + Profile->GetDefaultLaunchRole()->GetUATCommandLine().TrimStartAndEnd();
+	CommandLine.TrimStartAndEndInline();
+
+	return FText::FromString(CommandLine);
+}
+
+void FBasicProfileMarshaller::SetCommandLine( const FText& NewText, ETextCommit::Type InTextCommit )
+{
+	FString CommandLine = NewText.ToString();
+	Profile->SetAdditionalCommandLineParameters(CommandLine);
+	Profile->GetDefaultLaunchRole()->SetCommandLine(TEXT(""));
+
+	OnPropertyChanged();
+}
+
+void FBasicProfileMarshaller::SetDeployDeviceIDs(TArray<FString> DeployDeviceIDs)
+{
+	if (Profile->GetDeployedDeviceGroup() == nullptr)
+	{
+		Profile->SetDeployedDeviceGroup(Model->GetProfileManager()->AddNewDeviceGroup());
+	}
+
+	Profile->GetDeployedDeviceGroup()->RemoveAllDevices();
+	for (const FString& DeployDeviceID : DeployDeviceIDs)
+	{
+		Profile->GetDeployedDeviceGroup()->AddDevice(DeployDeviceID);
+	}
+	Model->UpdatedCookedPlatformsFromDeployDeviceProxy(Profile);
+
+	OnPropertyChanged();
+}
+
+TArray<FString> FBasicProfileMarshaller::GetDeployDeviceIDs() const
+{
+	return Profile->GetDeployedDeviceGroup()->GetDeviceIDs();
+}
+
+
+void FBasicProfileMarshaller::OnPropertyChanged()
+{
+	// do not save basic profiles - they're transient
+}
+
+
+
+TSharedPtr<ILaunchProfileMarshaller> FBasicProfileMarshallerFactory::TryCreateMarshaller( const ILauncherProfileRef& InProfile, const TSharedRef<CustomLaunchUI::FModel>& InModel )
+{
+	if (InModel->IsBasicLaunchProfile(InProfile))
+	{
+		return MakeShared<FBasicProfileMarshaller>(InProfile, InModel);
+	}
+
+	return nullptr;
+}
+
+
+
+#undef LOCTEXT_NAMESPACE
diff --git a/Engine/Plugins/Developer/CustomLaunchUI/Source/Private/Marshalling/BasicProfileMarshaller.h b/Engine/Plugins/Developer/CustomLaunchUI/Source/Private/Marshalling/BasicProfileMarshaller.h
new file mode 100644
--- /dev/null
+++ b/Engine/Plugins/Developer/CustomLaunchUI/Source/Private/Marshalling/BasicProfileMarshaller.h
@@ -0,0 +1,13 @@
+// Copyright Epic Games, Inc. All Rights Reserved.
+
+#pragma once
+
+#include "Marshalling/ILaunchProfileMarshaller.h"
+
+
+class FBasicProfileMarshallerFactory : public ILaunchProfileMarshallerFactory
+{
+public:
+	virtual TSharedPtr<ILaunchProfileMarshaller> TryCreateMarshaller( const ILauncherProfileRef& InProfile, const TSharedRef<CustomLaunchUI::FModel>& InModel ) override;
+
+};
diff --git a/Engine/Plugins/Developer/CustomLaunchUI/Source/Private/Marshalling/CustomProfileMarshaller.cpp b/Engine/Plugins/Developer/CustomLaunchUI/Source/Private/Marshalling/CustomProfileMarshaller.cpp
new file mode 100644
--- /dev/null
+++ b/Engine/Plugins/Developer/CustomLaunchUI/Source/Private/Marshalling/CustomProfileMarshaller.cpp
@@ -0,0 +1,1087 @@
+// Copyright Epic Games, Inc. All Rights Reserved.
+
+#include "Marshalling/CustomProfileMarshaller.h"
+#include "Model/CustomLaunchUIModel.h"
+#include "Widgets/Shared/SCustomLaunchProjectCombo.h"
+#include "Widgets/Shared/SCustomLaunchBuildTargetCombo.h"
+#include "Widgets/Shared/SCustomLaunchPlatformCombo.h"
+#include "Widgets/Shared/SCustomLaunchBuildConfigurationCombo.h"
+#include "Widgets/Shared/SCustomLaunchContentSchemeCombo.h"
+#include "Widgets/Shared/SCustomLaunchDeviceListView.h"
+#include "Widgets/Shared/SCustomLaunchMapListView.h"
+#include "Widgets/Input/SSegmentedControl.h"
+#include "Widgets/Text/STextBlock.h"
+#include "Widgets/Layout/SBox.h"
+#include "Widgets/SBoxPanel.h"
+#include "SResizeBox.h"
+#include "SSearchableComboBox.h"
+#include "SlateOptMacros.h"
+#include "PlatformInfo.h"
+#include "IDesktopPlatform.h"
+#include "GameProjectHelper.h"
+
+
+
+#define LOCTEXT_NAMESPACE "CustomProfileMarshaller"
+
+class FCustomProfileMarshaller : public ILaunchProfileMarshaller, public TSharedFromThis<FCustomProfileMarshaller>
+{
+public:
+	FCustomProfileMarshaller( const ILauncherProfileRef& Profile, const TSharedRef<CustomLaunchUI::FModel>& InModel );
+	virtual ~FCustomProfileMarshaller() = default;
+
+	virtual void Construct() override;
+
+	virtual FString GetName() const override
+	{
+		return TEXT("CustomProfile");
+	}
+
+	virtual FLaunchProfileTreeDataRef GetProfileTree() override
+	{
+		return TreeData;
+	}
+
+protected:
+	TArray<FString> CachedMapsToCook;
+	enum class EMapOption : uint8
+	{
+		Startup,
+		Selected,
+	};
+	EMapOption MapOption = EMapOption::Startup;
+
+	TArray<FString> CachedDeployDeviceIDs;
+	enum class EDeployDeviceOption : uint8
+	{
+		Default,
+		Selected,
+	};
+	EDeployDeviceOption DeployDeviceOption = EDeployDeviceOption::Default;
+
+
+	FString GetProjectPath() const;
+	void SetProjectName(FString ProjectPath);
+	bool HasProject() const;
+
+	FString GetBuildTarget() const;
+	void SetBuildTarget(FString BuildTarget);
+	TArray<EBuildTargetType> GetSupportedBuildTargetTypes() const;
+
+	void SetBuildConfiguration(EBuildConfiguration BuildConfiguration);
+	EBuildConfiguration GetBuildConfiguration() const;
+
+	void SetContentScheme(CustomLaunchUI::EContentScheme ContentScheme);
+	bool IsContentSchemeAvailable(CustomLaunchUI::EContentScheme) const;
+
+	FString GetCommandLine() const;
+	void SetCommandLine( const FString& NewCommandLine );
+
+	void SetSelectedPlatforms( TArray<FString> SelectedPlatforms );
+	TArray<FString> GetSelectedPlatforms() const;
+
+	void SetMapsToCook(TArray<FString> MapsToCook);
+	TArray<FString> GetMapsToCook() const;
+	EMapOption GetMapOption() const;
+	void SetMapOption( EMapOption MapOption );
+	float GetMapListHeight() const;
+	void SetMapListHeight( float NewHeight );
+	TSharedRef<SWidget> CreateMapListWidget();
+	float MapListHeight = 300.0f;
+
+
+	void SetDeployDeviceIDs(TArray<FString> DeployDeviceIDs);
+	TArray<FString> GetDeployDeviceIDs() const;
+	EDeployDeviceOption GetDeployDeviceOption() const;
+	void SetDeployDeviceOption( EDeployDeviceOption DeployDeviceOption );
+	float GetDeployDeviceListHeight() const;
+	void SetDeployDeviceListHeight( float NewHeight );
+	void OnDeviceRemoved(FString DeviceID);
+	TSharedRef<SWidget> CreateDeployDeviceWidget();
+	float DeployDeviceListHeight = 150.0f;
+
+	void SetBuildCode( bool bBuildCode );
+	bool GetBuildCode( ILauncherProfilePtr InProfile = nullptr ) const;
+
+	void SetBuildAlways( bool bBuildAlways );
+	bool GetBuildAlways( ILauncherProfilePtr InProfile = nullptr ) const;
+
+	void SetDeployToDevice( bool bDeployToDevice );
+	bool GetDeployToDevice( ILauncherProfilePtr InProfile = nullptr ) const;
+
+	void SetIsRunning( bool bRun );
+	bool GetIsRunning( ILauncherProfilePtr InProfile = nullptr ) const;
+
+	void OnInitialMapChanged( TSharedPtr<FString> InitialMap, ESelectInfo::Type );
+	TSharedPtr<FString> GetInitialMap() const;
+
+
+
+protected:
+	virtual void OnPropertyChanged() override;
+	void CacheStartupMapList();
+	TSharedRef<SWidget> OnGenerateComboWidget( TSharedPtr<FString> InComboString );
+
+
+	FLaunchProfileTreeDataRef TreeData;
+	const ILauncherProfileRef Profile;
+	const ILauncherProfileRef DefaultProfile;
+	const TSharedRef<CustomLaunchUI::FModel> Model;
+
+	CustomLaunchUI::EContentScheme ContentScheme;
+	TArray<TSharedPtr<FString>> CachedStartupMaps;
+	TSharedPtr<SSearchableComboBox> InitalMapCombo;
+	TSharedPtr<SCustomLaunchMapListView> MapListView;
+	TSharedPtr<SCustomLaunchDeviceListView> DeployDeviceListView;
+};
+
+
+FCustomProfileMarshaller::FCustomProfileMarshaller( const ILauncherProfileRef& InProfile, const TSharedRef<CustomLaunchUI::FModel>& InModel )
+	: TreeData( MakeShared<FLaunchProfileTreeData>(InProfile, InModel, this) )
+	, Profile(InProfile)
+	, DefaultProfile(InModel->GetDefaultCustomLaunchProfile())
+	, Model(InModel)
+{
+}
+
+BEGIN_SLATE_FUNCTION_BUILD_OPTIMIZATION
+void FCustomProfileMarshaller::Construct()
+{
+	auto ForPak = [this]()
+	{
+		return (ContentScheme == CustomLaunchUI::EContentScheme::PakFiles || ContentScheme == CustomLaunchUI::EContentScheme::DevelopmentPackage);
+	};
+
+	auto ForZenWS = [this]()
+	{
+		return (ContentScheme == CustomLaunchUI::EContentScheme::ZenPakStreaming);
+	};
+
+	auto ForCooked = [this]()
+	{
+		return (ContentScheme != CustomLaunchUI::EContentScheme::ZenPakStreaming && ContentScheme != CustomLaunchUI::EContentScheme::CookOnTheFly);
+	};
+
+	auto ForCode = [this]()
+	{
+		return GetBuildCode();
+	};
+
+	auto ForDeployment = [this]()
+	{
+		return GetDeployToDevice();
+	};
+
+	auto ForRun = [this]()
+	{
+		return GetIsRunning();
+	};
+
+	auto EmptyString = []()
+	{
+		return FString();
+	};
+
+
+
+	TreeData->AddHeading( LOCTEXT("GeneralSettingsHeading", "General Settings") )
+
+		.AddWidget( LOCTEXT("ProjectLabel", "Project"), 
+			SNew(SCustomLaunchProjectCombo)
+			.OnSelectionChanged(this, &FCustomProfileMarshaller::SetProjectName)
+			.SelectedProject(this, &FCustomProfileMarshaller::GetProjectPath)
+			.HasProject(this, &FCustomProfileMarshaller::HasProject)
+			.ShowAnyProjectOption(true)
+			.CurrentProjectOption(SCustomLaunchProjectCombo::ECurrentProjectOption::ActualProject)
+			.Font(FCoreStyle::Get().GetFontStyle("SmallFont"))
+		)
+
+		.AddWidget( LOCTEXT("TargetLabel", "Target"), 
+			{
+				.IsDefault = [this]()		{ return !Profile->HasBuildTargetSpecified() || Profile->GetBuildTarget().IsEmpty(); },
+				.SetToDefault = [this]()	{ SetBuildTarget(FString()); },
+				.IsEnabled = [this]()		{ return Profile->HasProjectSpecified(); },
+			},
+			SNew(SCustomLaunchBuildTargetCombo)
+			.OnSelectionChanged(this, &FCustomProfileMarshaller::SetBuildTarget)
+			.SelectedBuildTarget(this, &FCustomProfileMarshaller::GetBuildTarget)
+			.SupportedTargetTypes(this, &FCustomProfileMarshaller::GetSupportedBuildTargetTypes)
+			.SelectedProject(this, &FCustomProfileMarshaller::GetProjectPath)
+			.Font(FCoreStyle::Get().GetFontStyle("SmallFont"))
+		)
+
+		.AddWidget( LOCTEXT("PlatformLabel", "Platform"), 
+			SNew(SCustomLaunchPlatformCombo)
+			.SelectedPlatforms(this, &FCustomProfileMarshaller::GetSelectedPlatforms)
+			.OnSelectionChanged(this, &FCustomProfileMarshaller::SetSelectedPlatforms)
+			.BasicPlatformsOnly(CustomLaunchUI::bUseFriendlyBuildTargetSelection)
+			.Font(FCoreStyle::Get().GetFontStyle("SmallFont"))
+		)
+
+		.AddWidget( LOCTEXT("ConfigurationLabel", "Configuration"), 
+			{
+				.IsEnabled = ForCode,
+			},
+			SNew(SCustomLaunchBuildConfigurationCombo)
+			.OnSelectionChanged(this, &FCustomProfileMarshaller::SetBuildConfiguration)
+			.SelectedBuildConfiguration(this, &FCustomProfileMarshaller::GetBuildConfiguration)
+			.Font(FCoreStyle::Get().GetFontStyle("SmallFont"))
+		)
+
+		.AddWidget( LOCTEXT("ContentSchemeLabel", "Content Scheme"), 
+			SNew(SCustomLaunchContentSchemeCombo)
+			.OnSelectionChanged(this, &FCustomProfileMarshaller::SetContentScheme)
+			.SelectedContentScheme_Lambda( [this]() { return ContentScheme; } )
+			.IsContentSchemeAvailable(this, &FCustomProfileMarshaller::IsContentSchemeAvailable)
+			.Font(FCoreStyle::Get().GetFontStyle("SmallFont"))
+		)
+	;
+
+
+	TreeData->AddHeading( LOCTEXT("ContentSchemeHeading", "Content Scheme") )
+
+		.AddBoolean( LOCTEXT("CompressPakFilesLabel", "Compress Pak Files"),
+			{
+				.GetValue = [this]()			{ return Profile->IsCompressed(); },
+				.SetValue = [this](bool bValue)	{ Profile->SetCompressed(bValue); },
+				.GetDefaultValue = [this]()		{ return DefaultProfile->IsCompressed(); },
+				.IsVisible = ForPak,
+			}
+		)
+
+		.AddBoolean( LOCTEXT("UseIoStoreLabel", "Use Io Store"),
+			{
+				.GetValue = [this]()			{ return Profile->IsUsingIoStore(); },
+				.SetValue = [this](bool bValue)	{ Profile->SetUseIoStore(bValue); },
+				.GetDefaultValue = [this]()		{ return DefaultProfile->IsUsingIoStore(); },
+				.IsVisible = ForPak,
+			}
+		)
+
+		.AddBoolean( LOCTEXT("GenerateChunksLabel", "GenerateChunks"),
+			{
+				.GetValue = [this]()			{ return Profile->IsGeneratingChunks(); },
+				.SetValue = [this](bool bValue)	{ Profile->SetGenerateChunks(bValue); },
+				.GetDefaultValue = [this]()		{ return DefaultProfile->IsGeneratingChunks(); },
+				.IsVisible = ForPak,
+			}
+		)
+
+		.AddBoolean( LOCTEXT("ImportZenSnapshotLabel", "Import Best Match Zen Snapshot"),
+			{
+				.GetValue = [this]()			{ return Profile->IsImportingZenSnapshot(); },
+				.SetValue = [this](bool bValue)	{ Profile->SetImportingZenSnapshot(bValue); },
+				.GetDefaultValue = [this]()		{ return DefaultProfile->IsImportingZenSnapshot(); },
+				.IsVisible = ForCooked,
+			}
+		)
+
+		.AddDirectoryString( LOCTEXT("ZenPakStreamingPathLabel", "Zen Pak Streaming Path"),
+			{
+				.GetValue = [this]()				{ return Profile->GetZenPakStreamingPath(); },
+				.SetValue = [this](FString Value)	{ Profile->SetZenPakStreamingPath(Value); },
+				.GetDefaultValue = EmptyString,
+				.IsVisible = ForZenWS,
+			}
+		)
+
+		.AddBoolean( LOCTEXT("IncrementalCookLabel", "Incremental Cook"),
+			{
+				.GetValue = [this]()			{ return Profile->IsCookingIncrementally(); },
+				.SetValue = [this](bool bValue)	{ Profile->SetIncrementalCooking(bValue); },
+				.GetDefaultValue = [this]()		{ return DefaultProfile->IsCookingIncrementally(); },
+				.IsVisible = ForCooked,
+			}
+		)
+	;
+
+
+	TreeData->AddHeading( LOCTEXT("DirectoryHeading", "Directory") )
+
+		.AddDirectoryString( LOCTEXT("CustomStagingPathLabel", "Custom Stage Directory"),
+			{
+				.GetValue = [this]()				{ return Profile->GetPackageDirectory(); },
+				.SetValue = [this](FString Value)	{ Profile->SetPackageDirectory(Value); },
+				.GetDefaultValue = EmptyString,
+				.IsEnabled = ForCooked,
+			}
+		)
+
+		.AddBoolean( LOCTEXT("ArchiveBuildLabel", "Archive Build"),
+			{
+				.GetValue = [this]()			{ return Profile->IsArchiving(); },
+				.SetValue = [this](bool bValue)	{ Profile->SetArchive(bValue); },
+				.GetDefaultValue = [this]()		{ return DefaultProfile->IsArchiving(); },
+				.IsEnabled = ForCooked,
+			}
+		)
+
+		.AddDirectoryString( LOCTEXT("ArchivePathLabel", "Archive Directory"),
+			{
+				.GetValue = [this]()				{ return Profile->GetArchiveDirectory(); },
+				.SetValue = [this](FString Value)	{ Profile->SetArchiveDirectory(Value); },
+				.GetDefaultValue = EmptyString,
+				.IsEnabled = [this,ForCooked]()		{ return Profile->IsArchiving() && ForCooked(); }
+			}
+		)
+	;
+
+
+
+	TreeData->AddHeading( LOCTEXT("CookingHeading", "Maps And Cooking") )
+
+		.AddWidget( LOCTEXT("MapsToCookLabel", "Maps To Cook"), 
+			{
+				.IsEnabled = ForCooked,
+			},
+			CreateMapListWidget()
+		)
+
+		.AddString( LOCTEXT("AdditionalCookerOptionsLabel", "Additional Cooker Options"), 
+			{
+				.GetValue = [this]()				{ return Profile->GetCookOptions(); },
+				.SetValue = [this](FString Value)	{ Profile->SetCookOptions( Value ); },
+				.GetDefaultValue = EmptyString,
+				.IsEnabled = ForCooked,
+			}
+		)
+	;
+
+
+	TreeData->AddHeading( LOCTEXT("CodeHeading", "Code") )
+
+		.AddBoolean( LOCTEXT("BuildCodeLabel", "Build Code"),
+			{
+				.GetValue = [this]()			{ return GetBuildCode(); },
+				.SetValue = [this](bool bValue)	{ SetBuildCode(bValue); },
+				.GetDefaultValue = [this]()		{ return GetBuildCode(DefaultProfile); },
+			}
+		)
+
+		.AddBoolean( LOCTEXT("ForceBuildCodeLabel", "Build even if pre-build target exists"),
+			{
+				.GetValue = [this]()			{ return GetBuildAlways(); },
+				.SetValue = [this](bool bValue)	{ SetBuildAlways(bValue); },
+				.GetDefaultValue = [this]()		{ return GetBuildAlways(DefaultProfile); },
+				.IsEnabled = ForCode,
+			}
+		)
+
+		.AddBoolean( LOCTEXT("BuildUATLabel", "Build UAT"),
+			{
+				.GetValue = [this]()			{ return Profile->IsBuildingUAT(); },
+				.SetValue = [this](bool bValue)	{ Profile->SetBuildUAT(bValue); },
+				.GetDefaultValue = [this]()		{ return DefaultProfile->IsBuildingUAT(); },
+				.IsEnabled = ForCode,
+			}
+		)
+
+
+	;
+
+
+
+	TreeData->AddHeading( LOCTEXT("DeployAndRunHeading", "Deploy And Run") )
+
+		.AddBoolean( LOCTEXT("DeployLabel", "Deploy To Device"),
+			{
+				.GetValue = [this]()			{ return GetDeployToDevice(); },
+				.SetValue = [this](bool bValue)	{ SetDeployToDevice(bValue); },
+				.GetDefaultValue = [this]()		{ return GetDeployToDevice(DefaultProfile); },
+				.IsEnabled = [this]() { return (ContentScheme != CustomLaunchUI::EContentScheme::CookOnTheFly); }
+			}
+		)
+
+		.AddBoolean( LOCTEXT("DeployModifiedLabel", "Only Deploy Modified Content"),
+			{
+				.GetValue = [this]()			{ return Profile->IsDeployingIncrementally(); },
+				.SetValue = [this](bool bValue)	{ Profile->SetIncrementalDeploying(bValue); },
+				.GetDefaultValue = [this]()		{ return DefaultProfile->IsDeployingIncrementally(); },
+				.IsEnabled = ForDeployment,
+			}
+		)
+
+		.AddWidget( LOCTEXT("TargetDeviceLabel", "Target Device"), 
+			{
+				.IsEnabled = ForDeployment,
+			},
+			CreateDeployDeviceWidget()
+		)
+
+		.AddBoolean( LOCTEXT("RunLabel", "Run"),
+			{
+				.GetValue = [this]()			{ return GetIsRunning(); },
+				.SetValue = [this](bool bValue)	{ SetIsRunning(bValue); },
+				.GetDefaultValue = [this]()		{ return GetIsRunning(DefaultProfile); },
+			}
+		)
+
+		// todo: custom map picker 
+		.AddWidget( LOCTEXT("InitialMapLabel", "Initial Map"), 
+			{
+				.IsDefault = [this]()		{ return Profile->GetDefaultLaunchRole()->GetInitialMap().IsEmpty(); },
+				.SetToDefault = [this]()	{ Profile->GetDefaultLaunchRole()->SetInitialMap(FString()); InitalMapCombo->SetSelectedItem(GetInitialMap()); },
+				.IsEnabled = ForRun,
+			},
+			SAssignNew(InitalMapCombo, SSearchableComboBox)
+			.OptionsSource(&CachedStartupMaps)
+			.OnSelectionChanged( this, &FCustomProfileMarshaller::OnInitialMapChanged )
+			.OnGenerateWidget( this, &FCustomProfileMarshaller::OnGenerateComboWidget )
+			.Content()
+			[
+				SNew(STextBlock)
+				.Text_Lambda( [this]() { return FText::FromString( *GetInitialMap() ); } )
+				.Font(FCoreStyle::Get().GetFontStyle("SmallFont"))
+			]
+		)
+
+		.AddString( LOCTEXT("CommandLineLabel", "Command Line"), 
+			{
+				.GetValue = [this]()				{ return GetCommandLine(); },
+				.SetValue = [this](FString Value)	{ SetCommandLine(Value); },
+				.GetDefaultValue = EmptyString,
+				.IsEnabled = ForRun
+			}
+		)
+	;
+
+
+
+	MapOption = Profile->GetCookedMaps().Num() > 0 ? EMapOption::Selected : EMapOption::Startup;
+	DeployDeviceOption = Profile->GetDeployedDeviceGroup()->GetDeviceIDs().Num() > 0 ? EDeployDeviceOption::Selected : EDeployDeviceOption::Default;
+	ContentScheme = Model->DetermineProfileContentScheme(Profile);
+	CacheStartupMapList();
+}
+END_SLATE_FUNCTION_BUILD_OPTIMIZATION
+
+
+
+void FCustomProfileMarshaller::OnPropertyChanged()
+{
+	Model->GetProfileManager()->SaveJSONProfile(Profile);
+}
+
+
+void FCustomProfileMarshaller::CacheStartupMapList()
+{
+	CachedStartupMaps.Reset();
+	CachedStartupMaps.Add(MakeShared<FString>());
+	for ( const FString& Map : FGameProjectHelper::GetAvailableMaps(Profile->GetProjectBasePath(), Profile->SupportsEngineMaps(), true))
+	{
+		CachedStartupMaps.Add(MakeShared<FString>(Map));
+	}
+
+	InitalMapCombo->SetSelectedItem( GetInitialMap(), ESelectInfo::Direct );
+	InitalMapCombo->RefreshOptions();
+}
+
+TSharedRef<SWidget> FCustomProfileMarshaller::OnGenerateComboWidget( TSharedPtr<FString> InComboString )
+{
+	return SNew(STextBlock)
+		.Text(InComboString.IsValid() ? FText::FromString(*InComboString) : FText::GetEmpty() )
+		.Font(FCoreStyle::Get().GetFontStyle("SmallFont"));
+
+}
+
+
+
+void FCustomProfileMarshaller::SetSelectedPlatforms( TArray<FString> SelectedPlatforms )
+{
+	Profile->ClearCookedPlatforms();
+
+	if (CustomLaunchUI::bUseFriendlyBuildTargetSelection)
+	{
+		FTargetInfo BuildTargetInfo = CustomLaunchUI::FModel::GetBuildTargetInfo(GetBuildTarget(), GetProjectPath());
+
+		for (const FString& Platform : SelectedPlatforms)
+		{
+			Profile->AddCookedPlatform( CustomLaunchUI::FModel::GetBuildTargetPlatformName(Platform, BuildTargetInfo));
+		}
+	}
+	else
+	{
+		for (const FString& Platform : SelectedPlatforms)
+		{
+			Profile->AddCookedPlatform(Platform);
+		}
+	}
+	OnPropertyChanged();
+
+	DeployDeviceListView->OnSelectedPlatformChanged();
+}
+
+TArray<FString> FCustomProfileMarshaller::GetSelectedPlatforms() const
+{
+	TArray<FString> Platforms;
+
+	if (CustomLaunchUI::bUseFriendlyBuildTargetSelection)
+	{
+		for (const FString& Platform : Profile->GetCookedPlatforms())
+		{
+			Platforms.Add(CustomLaunchUI::FModel::GetVanillaPlatformName(Platform));
+		}
+	}
+	else
+	{
+		Platforms = Profile->GetCookedPlatforms();
+	}
+
+	return MoveTemp(Platforms);
+}
+
+
+FString FCustomProfileMarshaller::GetProjectPath() const
+{
+	if (Profile->HasProjectSpecified())
+	{
+		return Profile->GetProjectPath();
+	}
+
+	return FString();
+}
+
+void FCustomProfileMarshaller::SetProjectName(FString ProjectPath)
+{
+	Profile->SetProjectSpecified(!ProjectPath.IsEmpty());
+	Profile->SetProjectPath(ProjectPath);
+
+	// updating the project can reset the Taret
+	//Target = Profile->HasBuildTargetSpecified() && !Profile->GetBuildTarget().IsEmpty() ? Profile->GetBuildTarget() : FString();
+	//bHasTargetSpecified = Profile->HasBuildTargetSpecified();
+
+	if (CustomLaunchUI::bUseFriendlyBuildTargetSelection)
+	{
+		Model->UpdateCookedPlatformsFromBuildTarget(Profile);
+	}
+
+	OnPropertyChanged();
+
+	if (MapListView.IsValid())
+	{
+		MapListView->RefreshMapList();
+	}
+
+	CacheStartupMapList();
+}
+
+bool FCustomProfileMarshaller::HasProject() const
+{
+	return Profile->HasProjectSpecified();
+}
+
+
+FString FCustomProfileMarshaller::GetBuildTarget() const
+{
+	return Profile->GetBuildTarget();
+}
+
+void FCustomProfileMarshaller::SetBuildTarget(FString BuildTarget)
+{
+	Profile->SetBuildTargetSpecified(!BuildTarget.IsEmpty());
+	Profile->SetBuildTarget(BuildTarget);
+
+	if (CustomLaunchUI::bUseFriendlyBuildTargetSelection)
+	{
+		Model->UpdateCookedPlatformsFromBuildTarget(Profile);
+	}
+
+	OnPropertyChanged();
+}
+
+TArray<EBuildTargetType> FCustomProfileMarshaller::GetSupportedBuildTargetTypes() const
+{
+	TArray<EBuildTargetType> Result;
+
+	if (Profile->GetCookedPlatforms().Num() > 0)
+	{
+		const PlatformInfo::FTargetPlatformInfo* PlatformInfo = PlatformInfo::FindPlatformInfo(FName(Profile->GetCookedPlatforms()[0]));
+		if (PlatformInfo)
+		{
+			if (CustomLaunchUI::bUseFriendlyBuildTargetSelection)
+			{
+				Result.Add(PlatformInfo->VanillaInfo->PlatformType);
+
+				for (const PlatformInfo::FTargetPlatformInfo* PlatformFlavorInfo : PlatformInfo->VanillaInfo->Flavors)
+				{
+					Result.AddUnique(PlatformFlavorInfo->PlatformType);
+				}
+			}
+			else
+			{
+				Result.Add(PlatformInfo->PlatformType);
+			}
+		}
+	}
+
+	return MoveTemp(Result);
+}
+
+void FCustomProfileMarshaller::SetBuildConfiguration(EBuildConfiguration BuildConfiguration)
+{
+	Profile->SetBuildConfiguration(BuildConfiguration);
+	OnPropertyChanged();
+}
+
+EBuildConfiguration FCustomProfileMarshaller::GetBuildConfiguration() const
+{
+	return Profile->GetBuildConfiguration();
+}
+
+
+
+void FCustomProfileMarshaller::SetContentScheme(CustomLaunchUI::EContentScheme InContentScheme)
+{
+	ELauncherProfileDeploymentModes::Type DeploymentMode = GetDeployToDevice() ? ELauncherProfileDeploymentModes::CopyToDevice : ELauncherProfileDeploymentModes::DoNotDeploy;
+	Model->SetProfileContentScheme(InContentScheme, Profile, DeploymentMode );
+	OnPropertyChanged();
+
+	// refresh the cached content scheme in case the option that was selected is not available
+	ContentScheme = Model->DetermineProfileContentScheme(Profile);
+}
+
+bool FCustomProfileMarshaller::IsContentSchemeAvailable(CustomLaunchUI::EContentScheme InContentScheme) const
+{
+	// @fixme: check the project supports the content scheme etc
+	return true;
+}
+
+
+
+FString FCustomProfileMarshaller::GetCommandLine() const
+{
+	// get unified command line from these two fields. the first is presented in "Build" for old Project Launcher, and the latter is presented in "Launch" for old Project Launcher)
+	// when we save back to the profile, this will be stored just in the "Build" one for clarity (because multiple roles are not supported in old or new project launcher)
+	FString CommandLine = Profile->GetAdditionalCommandLineParameters().TrimStartAndEnd() + TEXT(" ") + Profile->GetDefaultLaunchRole()->GetUATCommandLine().TrimStartAndEnd();
+	CommandLine.TrimStartAndEndInline();
+
+	return CommandLine;
+}
+
+void FCustomProfileMarshaller::SetCommandLine( const FString& NewCommandLine )
+{
+	Profile->SetAdditionalCommandLineParameters(NewCommandLine);
+	Profile->GetDefaultLaunchRole()->SetCommandLine(TEXT(""));
+	OnPropertyChanged();
+}
+
+
+
+
+BEGIN_SLATE_FUNCTION_BUILD_OPTIMIZATION
+TSharedRef<SWidget> FCustomProfileMarshaller::CreateMapListWidget()
+{
+	MapListView = 
+		SNew(SCustomLaunchMapListView)
+		.OnSelectionChanged(this, &FCustomProfileMarshaller::SetMapsToCook)
+		.SelectedMaps(this, &FCustomProfileMarshaller::GetMapsToCook)
+		.ProjectPath(this, &FCustomProfileMarshaller::GetProjectPath)
+	;
+
+	return SNew(SVerticalBox)
+
+		// map option controls
+		+SVerticalBox::Slot()
+		.AutoHeight()
+		.Padding(0,2)
+		[
+			SNew(SHorizontalBox)
+
+			// map option
+			+SHorizontalBox::Slot()
+			.Padding(0, 0)
+			.AutoWidth()
+			.VAlign(VAlign_Center)
+			[
+				SNew(SSegmentedControl<EMapOption>)
+				.Value( this, &FCustomProfileMarshaller::GetMapOption)
+				.OnValueChanged(this, &FCustomProfileMarshaller::SetMapOption)
+
+				+SSegmentedControl<EMapOption>::Slot(EMapOption::Startup)
+				.HAlign(HAlign_Center)
+				.VAlign(VAlign_Center)
+				[
+					SNew(STextBlock)
+					.Text(LOCTEXT("StartupMapsLabel", "Startup Maps"))
+					.Font(FCoreStyle::Get().GetFontStyle("SmallFont"))
+				]
+
+				+SSegmentedControl<EMapOption>::Slot(EMapOption::Selected)
+				.HAlign(HAlign_Center)
+				.VAlign(VAlign_Center)
+				[
+					SNew(STextBlock)
+					.Text(LOCTEXT("SelectedMapsLabel", "Selected Maps"))
+					.Font(FCoreStyle::Get().GetFontStyle("SmallFont"))
+				]
+			]
+
+			// map selector controls (search etc)
+			+SHorizontalBox::Slot()
+			.Padding(8, 0)
+			.FillWidth(1)
+			.VAlign(VAlign_Center)
+			[
+				SNew(SBox)
+				.Visibility_Lambda([this](){ return GetMapOption() == EMapOption::Selected  ? EVisibility::Visible : EVisibility::Collapsed; } )
+				[
+					MapListView->MakeControlsWidget()
+				]
+			]
+		]
+
+		// map list
+		+SVerticalBox::Slot()
+		.FillHeight(1)
+		[
+			SNew(SVerticalResizeBox)
+			.Visibility_Lambda([this](){ return (GetMapOption() == EMapOption::Selected) ? EVisibility::Visible : EVisibility::Collapsed; } )
+			.HandleHeight(4)
+			.ContentHeight(this, &FCustomProfileMarshaller::GetMapListHeight)
+			.ContentHeightChanged(this, &FCustomProfileMarshaller::SetMapListHeight)
+			.HandleColor( FAppStyle::Get().GetSlateColor("Colors.Secondary").GetSpecifiedColor() )
+			[
+				MapListView.ToSharedRef()
+			]
+		]
+	;
+}
+END_SLATE_FUNCTION_BUILD_OPTIMIZATION
+
+
+void FCustomProfileMarshaller::SetMapsToCook(TArray<FString> MapsToCook)
+{
+	Profile->ClearCookedMaps(); 
+	for (const FString& Map : MapsToCook)
+	{
+		Profile->AddCookedMap(Map);
+	}
+
+	OnPropertyChanged();
+}
+
+TArray<FString> FCustomProfileMarshaller::GetMapsToCook() const
+{
+	return Profile->GetCookedMaps();
+}
+
+float FCustomProfileMarshaller::GetMapListHeight() const
+{
+	return MapListHeight;
+}
+
+void FCustomProfileMarshaller::SetMapListHeight( float NewHeight )
+{
+	static const float MinMapListHeight = 100.0f;
+
+	MapListHeight = FMath::Max(NewHeight, MinMapListHeight);
+}
+
+FCustomProfileMarshaller::EMapOption FCustomProfileMarshaller::GetMapOption() const
+{
+	return MapOption;
+}
+
+void FCustomProfileMarshaller::SetMapOption( EMapOption NewMapOption )
+{
+	bool bShow = (NewMapOption == EMapOption::Selected);
+
+	MapOption = NewMapOption;
+
+	if (bShow)
+	{
+		// restore the cooked maps again, if any
+		if (CachedMapsToCook.Num() > 0 && Profile->GetCookedMaps().Num() == 0)
+		{
+			SetMapsToCook(CachedMapsToCook);
+			CachedMapsToCook.Reset();
+		}
+	}
+	else
+	{
+		// to set the 'cook startup maps only', its necessary to remove all the cooked maps - take a copy of the values to allow it to be restored
+		CachedMapsToCook = Profile->GetCookedMaps();
+		SetMapsToCook(TArray<FString>());
+	}
+
+	OnPropertyChanged();
+
+	MapListView->RefreshMapList();
+}
+
+
+
+
+BEGIN_SLATE_FUNCTION_BUILD_OPTIMIZATION
+TSharedRef<SWidget> FCustomProfileMarshaller::CreateDeployDeviceWidget()
+{
+	return SNew(SVerticalBox)
+
+		// device picker options
+		+SVerticalBox::Slot()
+		.AutoHeight()
+		.Padding(0,2)
+		[
+			SNew(SHorizontalBox)
+
+			+SHorizontalBox::Slot()
+			.Padding(0, 0)
+			.AutoWidth()
+			.VAlign(VAlign_Center)
+			[
+				SNew(SSegmentedControl<EDeployDeviceOption>)
+				.Value( this, &FCustomProfileMarshaller::GetDeployDeviceOption)
+				.OnValueChanged(this, &FCustomProfileMarshaller::SetDeployDeviceOption)
+
+				+SSegmentedControl<EDeployDeviceOption>::Slot(EDeployDeviceOption::Default)
+				.HAlign(HAlign_Center)
+				.VAlign(VAlign_Center)
+				[
+					SNew(STextBlock)
+					.Text(LOCTEXT("DefaultDeviceLabel", "Default Device"))
+					.Font(FCoreStyle::Get().GetFontStyle("SmallFont"))
+				]
+
+				+SSegmentedControl<EDeployDeviceOption>::Slot(EDeployDeviceOption::Selected)
+				.HAlign(HAlign_Center)
+				.VAlign(VAlign_Center)
+				[
+					SNew(STextBlock)
+					.Text(LOCTEXT("SelectedDevicesLabel", "Selected Devices"))
+					.Font(FCoreStyle::Get().GetFontStyle("SmallFont"))
+				]
+			]
+		]
+
+		// device picker list
+		+SVerticalBox::Slot()
+		.FillHeight(1)
+		[
+			SNew(SVerticalResizeBox)
+			.Visibility_Lambda([this](){ return (GetDeployDeviceOption() == EDeployDeviceOption::Selected) ? EVisibility::Visible : EVisibility::Collapsed; } )
+			.HandleHeight(4)
+			.ContentHeight(this, &FCustomProfileMarshaller::GetDeployDeviceListHeight)
+			.ContentHeightChanged(this, &FCustomProfileMarshaller::SetDeployDeviceListHeight)
+			.HandleColor( FAppStyle::Get().GetSlateColor("Colors.Secondary").GetSpecifiedColor() )
+			[
+				SAssignNew(DeployDeviceListView, SCustomLaunchDeviceListView)
+				.OnDeviceRemoved( this, &FCustomProfileMarshaller::OnDeviceRemoved)
+				.OnSelectionChanged(this, &FCustomProfileMarshaller::SetDeployDeviceIDs)
+				.SelectedDevices(this, &FCustomProfileMarshaller::GetDeployDeviceIDs)
+				.Platforms(this, &FCustomProfileMarshaller::GetSelectedPlatforms)
+			]
+		]
+	;
+}
+END_SLATE_FUNCTION_BUILD_OPTIMIZATION
+
+
+void FCustomProfileMarshaller::SetDeployDeviceIDs(TArray<FString> DeployDeviceIDs)
+{
+	if (Profile->GetDeployedDeviceGroup() == nullptr)
+	{
+		Profile->SetDeployedDeviceGroup(Model->GetProfileManager()->AddNewDeviceGroup());
+	}
+
+	Profile->GetDeployedDeviceGroup()->RemoveAllDevices(); 
+	for ( const FString& DeviceID : DeployDeviceIDs )
+	{
+		Profile->GetDeployedDeviceGroup()->AddDevice(DeviceID);
+	}
+	// Model->UpdatedCookedPlatformsFromDeployDeviceProxy(Profile); // fixme: needed?
+
+	OnPropertyChanged();
+}
+
+TArray<FString> FCustomProfileMarshaller::GetDeployDeviceIDs() const
+{
+	return Profile->GetDeployedDeviceGroup()->GetDeviceIDs();
+}
+
+float FCustomProfileMarshaller::GetDeployDeviceListHeight() const
+{
+	return DeployDeviceListHeight;
+}
+
+void FCustomProfileMarshaller::SetDeployDeviceListHeight( float NewHeight )
+{
+	static const float MinDeployDeviceListHeight = 100.0f;
+
+	DeployDeviceListHeight = FMath::Max(NewHeight, MinDeployDeviceListHeight);
+}
+
+void FCustomProfileMarshaller::OnDeviceRemoved(FString DeviceID)
+{
+	CachedDeployDeviceIDs.Remove(DeviceID);
+}
+
+FCustomProfileMarshaller::EDeployDeviceOption FCustomProfileMarshaller::GetDeployDeviceOption() const
+{
+	return DeployDeviceOption;
+}
+
+void FCustomProfileMarshaller::SetDeployDeviceOption( EDeployDeviceOption NewDeployDeviceOption )
+{
+	bool bShow = (NewDeployDeviceOption == EDeployDeviceOption::Selected);
+
+	DeployDeviceOption = NewDeployDeviceOption;
+
+	if (bShow)
+	{
+		// restore the deployed device list again, if any
+		if (CachedDeployDeviceIDs.Num() > 0 && Profile->GetDeployedDeviceGroup()->GetDeviceIDs().Num() == 0)
+		{
+			SetDeployDeviceIDs(Profile->GetDeployedDeviceGroup()->GetDeviceIDs());
+			CachedDeployDeviceIDs.Reset();
+		}
+	}
+	else
+	{
+		// to set the 'default' deploy option, its necessary to remove all the devices - take a copy of the values to allow it to be restored
+		CachedDeployDeviceIDs = Profile->GetDeployedDeviceGroup()->GetDeviceIDs();
+		SetDeployDeviceIDs(TArray<FString>());
+	}
+
+	OnPropertyChanged();
+
+	DeployDeviceListView->RefreshDeviceList();
+}
+
+
+void FCustomProfileMarshaller::SetBuildCode( bool bBuildCode )
+{
+	if (!bBuildCode)
+	{
+		Profile->SetBuildMode(ELauncherProfileBuildModes::DoNotBuild);
+	}
+	else if (GetBuildAlways(Profile))
+	{
+		Profile->SetBuildMode(ELauncherProfileBuildModes::Build);
+	}
+	else
+	{
+		Profile->SetBuildMode(ELauncherProfileBuildModes::Auto);
+	}
+
+	OnPropertyChanged();
+}
+
+bool FCustomProfileMarshaller::GetBuildCode(ILauncherProfilePtr InProfile) const
+{
+	if (InProfile == nullptr)
+	{
+		InProfile = Profile;
+	}
+	return (InProfile->GetBuildMode() != ELauncherProfileBuildModes::DoNotBuild);
+}
+
+void FCustomProfileMarshaller::SetBuildAlways( bool bBuildAlways )
+{
+	if (!GetBuildCode())
+	{
+		Profile->SetBuildMode(ELauncherProfileBuildModes::DoNotBuild);
+	}
+	else if (bBuildAlways)
+	{
+		Profile->SetBuildMode(ELauncherProfileBuildModes::Build);
+	}
+	else
+	{
+		Profile->SetBuildMode(ELauncherProfileBuildModes::Auto);
+	}
+
+	OnPropertyChanged();
+}
+
+bool FCustomProfileMarshaller::GetBuildAlways(ILauncherProfilePtr InProfile) const
+{
+	if (InProfile == nullptr)
+	{
+		InProfile = Profile;
+	}
+	return (InProfile->GetBuildMode() == ELauncherProfileBuildModes::Build);
+}
+
+
+void FCustomProfileMarshaller::SetDeployToDevice( bool bDeployToDevice )
+{
+	if (ContentScheme != CustomLaunchUI::EContentScheme::CookOnTheFly)
+	{
+		Profile->SetDeploymentMode(bDeployToDevice ? ELauncherProfileDeploymentModes::CopyToDevice : ELauncherProfileDeploymentModes::DoNotDeploy );
+		OnPropertyChanged();
+	}
+}
+bool FCustomProfileMarshaller::GetDeployToDevice(ILauncherProfilePtr InProfile) const
+{
+	if (InProfile == nullptr)
+	{
+		InProfile = Profile;
+	}
+	return InProfile->GetDeploymentMode() != ELauncherProfileDeploymentModes::DoNotDeploy;
+}
+
+
+void FCustomProfileMarshaller::SetIsRunning( bool bRun )
+{
+	Profile->SetLaunchMode(bRun ? ELauncherProfileLaunchModes::DefaultRole : ELauncherProfileLaunchModes::DoNotLaunch);
+}
+
+bool FCustomProfileMarshaller::GetIsRunning( ILauncherProfilePtr InProfile) const
+{
+	if (InProfile == nullptr)
+	{
+		InProfile = Profile;
+	}
+	
+	return InProfile->GetLaunchMode() != ELauncherProfileLaunchModes::DoNotLaunch;
+}
+
+void FCustomProfileMarshaller::OnInitialMapChanged( TSharedPtr<FString> InitialMap, ESelectInfo::Type )
+{
+	if (InitialMap.IsValid())
+	{
+		Profile->GetDefaultLaunchRole()->SetInitialMap(*InitialMap);
+	}
+	else
+	{
+		Profile->GetDefaultLaunchRole()->SetInitialMap(FString());
+	}
+
+	OnPropertyChanged();
+}
+
+TSharedPtr<FString> FCustomProfileMarshaller::GetInitialMap() const
+{
+	FString InitialMap = Profile->GetDefaultLaunchRole()->GetInitialMap();
+	const TSharedPtr<FString>* FoundMapPtr = CachedStartupMaps.FindByPredicate( [InitialMap]( const TSharedPtr<FString>& Map )
+	{
+		return Map.IsValid() && (InitialMap == *Map);
+	});
+
+	if (FoundMapPtr)
+	{
+		return *FoundMapPtr;
+	}
+	else
+	{
+		return MakeShared<FString>();
+	}
+}
+
+
+
+
+
+
+TSharedPtr<ILaunchProfileMarshaller> FCustomProfileMarshallerFactory::TryCreateMarshaller( const ILauncherProfileRef& InProfile, const TSharedRef<CustomLaunchUI::FModel>& InModel )
+{
+	if (!InModel->IsBasicLaunchProfile(InProfile) && !InModel->IsAdvancedProfile(InProfile))
+	{
+		return MakeShared<FCustomProfileMarshaller>(InProfile, InModel);
+	}
+
+	return nullptr;
+}
+
+
+
+#undef LOCTEXT_NAMESPACE
diff --git a/Engine/Plugins/Developer/CustomLaunchUI/Source/Private/Marshalling/CustomProfileMarshaller.h b/Engine/Plugins/Developer/CustomLaunchUI/Source/Private/Marshalling/CustomProfileMarshaller.h
new file mode 100644
--- /dev/null
+++ b/Engine/Plugins/Developer/CustomLaunchUI/Source/Private/Marshalling/CustomProfileMarshaller.h
@@ -0,0 +1,13 @@
+// Copyright Epic Games, Inc. All Rights Reserved.
+
+#pragma once
+
+#include "Marshalling/ILaunchProfileMarshaller.h"
+
+
+class FCustomProfileMarshallerFactory : public ILaunchProfileMarshallerFactory
+{
+public:
+	virtual TSharedPtr<ILaunchProfileMarshaller> TryCreateMarshaller( const ILauncherProfileRef& Profile, const TSharedRef<CustomLaunchUI::FModel>& InModel ) override;
+
+};
diff --git a/Engine/Plugins/Developer/CustomLaunchUI/Source/Private/Marshalling/ILaunchProfileMarshaller.cpp b/Engine/Plugins/Developer/CustomLaunchUI/Source/Private/Marshalling/ILaunchProfileMarshaller.cpp
new file mode 100644
--- /dev/null
+++ b/Engine/Plugins/Developer/CustomLaunchUI/Source/Private/Marshalling/ILaunchProfileMarshaller.cpp
@@ -0,0 +1,82 @@
+// Copyright Epic Games, Inc. All Rights Reserved.
+
+#include "Marshalling/ILaunchProfileMarshaller.h"
+
+
+static TArray<TSharedPtr<ILaunchProfileMarshallerFactory>> GMarshallerFactories;
+
+
+
+class FEmptyProfileMarshaller : public ILaunchProfileMarshaller
+{
+public:
+	FEmptyProfileMarshaller(const ILauncherProfilePtr& InProfile, TSharedRef<CustomLaunchUI::FModel> InModel)
+		: TreeData(MakeShared<FLaunchProfileTreeData>(InProfile, InModel, this))
+	{
+	}
+	virtual ~FEmptyProfileMarshaller() = default;
+
+	virtual void Construct() override
+	{
+	}
+
+	virtual FString GetName() const override
+	{
+		return TEXT("Empty");
+	}
+
+	virtual FLaunchProfileTreeDataRef GetProfileTree() override
+	{
+		return TreeData;
+	}
+
+	virtual void OnPropertyChanged() override
+	{
+	}
+
+protected:
+	FLaunchProfileTreeDataRef TreeData;
+};
+
+
+
+
+
+void ILaunchProfileMarshallerFactory::Register( TSharedRef<ILaunchProfileMarshallerFactory> Marshaller )
+{
+	GMarshallerFactories.Add(Marshaller);
+	GMarshallerFactories.Sort( []( TSharedPtr<ILaunchProfileMarshallerFactory> A, TSharedPtr<ILaunchProfileMarshallerFactory> B)
+	{
+		return A->GetPriority() > B->GetPriority();
+	});
+
+}
+
+void ILaunchProfileMarshallerFactory::Unregister( TSharedRef<ILaunchProfileMarshallerFactory> Marshaller )
+{
+	GMarshallerFactories.Remove(Marshaller);
+}
+
+
+TSharedRef<ILaunchProfileMarshaller> ILaunchProfileMarshallerFactory::Marshall( const ILauncherProfilePtr& InProfile, TSharedRef<CustomLaunchUI::FModel> InModel )
+{
+	if (InProfile.IsValid())
+	{
+		// find a marshaller that can handle this profile
+		for (const TSharedPtr<ILaunchProfileMarshallerFactory>& Factory : GMarshallerFactories)
+		{
+			TSharedPtr<ILaunchProfileMarshaller> Marshaller = Factory->TryCreateMarshaller(InProfile.ToSharedRef(), InModel);
+			if (Marshaller.IsValid())
+			{
+				Marshaller->Construct();
+				return Marshaller.ToSharedRef();
+			}
+		}
+	}
+
+	// cannot marshall
+	TSharedRef<ILaunchProfileMarshaller> EmptyMarshaller = MakeShared<FEmptyProfileMarshaller>(InProfile, InModel);
+	EmptyMarshaller->Construct();
+	return EmptyMarshaller;
+}
+
diff --git a/Engine/Plugins/Developer/CustomLaunchUI/Source/Private/Marshalling/ILaunchProfileMarshaller.h b/Engine/Plugins/Developer/CustomLaunchUI/Source/Private/Marshalling/ILaunchProfileMarshaller.h
new file mode 100644
--- /dev/null
+++ b/Engine/Plugins/Developer/CustomLaunchUI/Source/Private/Marshalling/ILaunchProfileMarshaller.h
@@ -0,0 +1,29 @@
+// Copyright Epic Games, Inc. All Rights Reserved.
+
+#pragma once
+
+#include "Marshalling/LaunchProfileTreeData.h"
+#include "ILauncherProfile.h"
+
+class ILaunchProfileMarshaller
+{
+public:
+	virtual void Construct() = 0;
+	virtual FLaunchProfileTreeDataRef GetProfileTree() = 0;
+	virtual FString GetName() const = 0;
+	virtual void OnPropertyChanged() = 0;
+};
+
+
+class ILaunchProfileMarshallerFactory
+{
+public:
+	static void Register( TSharedRef<ILaunchProfileMarshallerFactory> Marshaller );
+	static void Unregister( TSharedRef<ILaunchProfileMarshallerFactory> Marshaller );
+	static TSharedRef<ILaunchProfileMarshaller> Marshall( const ILauncherProfilePtr& InProfile, TSharedRef<CustomLaunchUI::FModel> InModel );
+
+	virtual TSharedPtr<ILaunchProfileMarshaller> TryCreateMarshaller( const ILauncherProfileRef& Profile, const TSharedRef<CustomLaunchUI::FModel>& InModel ) = 0;
+	virtual int GetPriority() const { return 0; }
+};
+
+
diff --git a/Engine/Plugins/Developer/CustomLaunchUI/Source/Private/Marshalling/LaunchProfileTreeData.cpp b/Engine/Plugins/Developer/CustomLaunchUI/Source/Private/Marshalling/LaunchProfileTreeData.cpp
new file mode 100644
--- /dev/null
+++ b/Engine/Plugins/Developer/CustomLaunchUI/Source/Private/Marshalling/LaunchProfileTreeData.cpp
@@ -0,0 +1,254 @@
+// Copyright Epic Games, Inc. All Rights Reserved.
+
+#include "Marshalling/LaunchProfileTreeData.h"
+#include "Marshalling/ILaunchProfileMarshaller.h"
+#include "Styling/CustomLaunchUIStyle.h"
+#include "Widgets/Input/SCheckBox.h"
+#include "Widgets/Input/SButton.h"
+#include "Widgets/Input/SEditableTextBox.h"
+#include "Widgets/Input/SMultiLineEditableTextBox.h"
+#include "Widgets/Images/SImage.h"
+#include "Widgets/SBoxPanel.h"
+#include "Layout/Visibility.h"
+#include "Framework/Application/SlateApplication.h"
+#include "DesktopPlatformModule.h"
+
+
+#define LOCTEXT_NAMESPACE "SCustomLaunchCustomProfileEditor"
+
+
+FLaunchProfileTreeData::FLaunchProfileTreeData(ILauncherProfilePtr InProfile, TSharedRef<CustomLaunchUI::FModel> InModel, ILaunchProfileMarshaller* InMarshaller) 
+	: Profile(InProfile)
+	, Model(InModel)
+	, Marshaller(InMarshaller)
+{
+	check(Marshaller);
+}
+
+FLaunchProfileTreeNode& FLaunchProfileTreeData::AddHeading( FText InName )
+{
+	FLaunchProfileTreeNodePtr TreeNode = MakeShared<FLaunchProfileTreeNode>(this);
+	TreeNode->Name = InName;
+	Nodes.Add(TreeNode);
+
+	return *TreeNode.Get();
+}
+
+
+
+FLaunchProfileTreeNode::FLaunchProfileTreeNode( const FLaunchProfileTreeData* InTreeData ) 
+	: TreeData(InTreeData)
+{
+}
+
+FLaunchProfileTreeNode& FLaunchProfileTreeNode::AddWidget( FText InName, FCallbacks&& InWidgetCallbacks, TSharedRef<SWidget> InValueWidget )
+{
+	FLaunchProfileTreeNodePtr TreeNode = MakeShared<FLaunchProfileTreeNode>(TreeData);
+	TreeNode->Name = InName;
+	TreeNode->Widget = InValueWidget;
+	TreeNode->Callbacks = InWidgetCallbacks;
+	Children.Add(TreeNode);
+
+	return *this;
+}
+
+FLaunchProfileTreeNode& FLaunchProfileTreeNode::AddWidget( FText InName, TSharedRef<SWidget> InValueWidget )
+{
+	return AddWidget( InName, {}, InValueWidget );
+}
+
+
+FLaunchProfileTreeNode& FLaunchProfileTreeNode::AddBoolean( FText InName, FBooleanCallbacks&& BooleanCallbacks )
+{
+	check(BooleanCallbacks.GetValue);
+	check(BooleanCallbacks.SetValue);
+
+	FCallbacks WidgetCallbacks
+	{
+		.IsVisible = BooleanCallbacks.IsVisible,
+		.IsEnabled = BooleanCallbacks.IsEnabled,
+	};
+	if (BooleanCallbacks.GetDefaultValue != nullptr)
+	{
+		WidgetCallbacks.IsDefault = [BooleanCallbacks]()
+		{
+			return BooleanCallbacks.GetValue() == BooleanCallbacks.GetDefaultValue();
+		};
+
+		WidgetCallbacks.SetToDefault = [BooleanCallbacks]()
+		{
+			BooleanCallbacks.SetValue(BooleanCallbacks.GetDefaultValue());
+			// nb. OnPropertyChanged is called via SCustomLaunchWidgetTableRow's set to default
+		};
+	};
+
+	auto SetCheckState = [this, BooleanCallbacks]( ECheckBoxState CheckState )
+	{
+		BooleanCallbacks.SetValue( (CheckState == ECheckBoxState::Checked) );
+		TreeData->Marshaller->OnPropertyChanged();
+	};
+
+	auto GetCheckState = [BooleanCallbacks]()
+	{
+		return BooleanCallbacks.GetValue() ? ECheckBoxState::Checked : ECheckBoxState::Unchecked;
+	};
+
+	return AddWidget( InName,
+		MoveTemp(WidgetCallbacks),
+		SNew(SCheckBox)
+		.OnCheckStateChanged_Lambda(SetCheckState)
+		.IsChecked_Lambda(GetCheckState)
+	);
+}
+
+FLaunchProfileTreeNode& FLaunchProfileTreeNode::AddString( FText InName, FStringCallbacks&& StringCallbacks )
+{
+	check(StringCallbacks.GetValue);
+	check(StringCallbacks.SetValue);
+
+	FCallbacks WidgetCallbacks
+	{
+		.IsVisible = StringCallbacks.IsVisible,
+		.IsEnabled = StringCallbacks.IsEnabled,
+	};
+	if (StringCallbacks.GetDefaultValue != nullptr)
+	{
+		WidgetCallbacks.IsDefault = [StringCallbacks]()
+		{
+			return StringCallbacks.GetValue() == StringCallbacks.GetDefaultValue();
+		};
+
+		WidgetCallbacks.SetToDefault = [StringCallbacks]()
+		{
+			StringCallbacks.SetValue(StringCallbacks.GetDefaultValue());
+			// nb. OnPropertyChanged is called via SCustomLaunchWidgetTableRow's set to default
+		};
+	};
+
+	auto SetString = [this, StringCallbacks]( const FText& InText, ETextCommit::Type InTextCommit )
+	{
+		StringCallbacks.SetValue( InText.ToString() );
+		TreeData->Marshaller->OnPropertyChanged();
+	};
+
+	auto GetString = [StringCallbacks]()
+	{
+		return FText::FromString(StringCallbacks.GetValue());
+	};
+
+	return AddWidget( InName,
+		MoveTemp(WidgetCallbacks),
+		SNew(SMultiLineEditableTextBox)
+		.AllowMultiLine(false)
+		.AutoWrapText(true)
+		.Font(FCoreStyle::Get().GetFontStyle("SmallFont"))
+		.Text_Lambda( GetString )
+		.OnTextCommitted_Lambda( SetString )
+	);
+
+}
+
+FLaunchProfileTreeNode& FLaunchProfileTreeNode::AddDirectoryString( FText InName, FStringCallbacks&& StringCallbacks )
+{
+	check(StringCallbacks.GetValue);
+	check(StringCallbacks.SetValue);
+
+	FCallbacks WidgetCallbacks
+	{
+		.IsVisible = StringCallbacks.IsVisible,
+		.IsEnabled = StringCallbacks.IsEnabled,
+	};
+	if (StringCallbacks.GetDefaultValue != nullptr)
+	{
+		WidgetCallbacks.IsDefault = [StringCallbacks]()
+		{
+			return StringCallbacks.GetValue() == StringCallbacks.GetDefaultValue();
+		};
+
+		WidgetCallbacks.SetToDefault = [StringCallbacks]()
+		{
+			StringCallbacks.SetValue(StringCallbacks.GetDefaultValue());
+			// nb. OnPropertyChanged is called via SCustomLaunchWidgetTableRow's set to default
+		};
+	};
+
+	auto SetString = [this, StringCallbacks]( const FText& InText, ETextCommit::Type InTextCommit )
+	{
+		StringCallbacks.SetValue( InText.ToString() );
+		TreeData->Marshaller->OnPropertyChanged();
+	};
+
+	auto GetString = [StringCallbacks]()
+	{
+		return FText::FromString(StringCallbacks.GetValue());
+	};
+
+	auto OnBrowse = [this,StringCallbacks]()
+	{
+		FString InitialDirectory = StringCallbacks.GetValue();
+		if (InitialDirectory.IsEmpty())
+		{
+			InitialDirectory = this->TreeData->Profile->GetProjectBasePath();
+		}
+		if (!InitialDirectory.IsEmpty() && FPaths::IsRelative(InitialDirectory))
+		{
+			InitialDirectory = FPaths::Combine( FPaths::RootDir(), InitialDirectory );
+		}
+
+		const void* ParentWindowHandle = FSlateApplication::Get().FindBestParentWindowHandleForDialogs(nullptr);
+		FString OutDirectory;
+		if (FDesktopPlatformModule::Get()->OpenDirectoryDialog(ParentWindowHandle, LOCTEXT("FolderDialogTitle", "Choose a directory").ToString(), InitialDirectory, OutDirectory))
+		{
+			if (FPaths::IsUnderDirectory(OutDirectory, FPaths::RootDir()))
+			{
+				FPaths::MakePathRelativeTo(OutDirectory, *FPaths::RootDir());
+			}
+
+			StringCallbacks.SetValue( OutDirectory );
+			TreeData->Marshaller->OnPropertyChanged();
+		}
+
+		return FReply::Handled();
+	};
+
+
+	return AddWidget( InName,
+		MoveTemp(WidgetCallbacks),
+		SNew(SHorizontalBox)
+
+		// path field
+		+SHorizontalBox::Slot()
+		.FillWidth(1)
+		.VAlign(VAlign_Center)
+		[
+			SNew(SEditableTextBox)
+			.Font(FCoreStyle::Get().GetFontStyle("SmallFont"))
+			.Text_Lambda( GetString )
+			.OnTextCommitted_Lambda( SetString )
+		]
+
+		// browse button
+		+SHorizontalBox::Slot()
+		.AutoWidth()
+		.Padding(FMargin(4, 0, 0, 0))
+		.VAlign(VAlign_Center)
+		[
+			SNew(SButton)
+			.ButtonStyle(FCustomLaunchUIStyle::Get(), "HoverHintOnly")
+			.ToolTipText(LOCTEXT("DirBrowseTip", "Browse for a folder"))
+			.OnClicked_Lambda(OnBrowse)
+			.ContentPadding(2)
+			[
+				SNew(SImage)
+				.Image(FCustomLaunchUIStyle::Get().GetBrush("PathPickerButton"))
+				.ColorAndOpacity(FSlateColor::UseForeground())
+			]
+		]
+	);
+}
+
+
+
+
+
+#undef LOCTEXT_NAMESPACE
diff --git a/Engine/Plugins/Developer/CustomLaunchUI/Source/Private/Marshalling/LaunchProfileTreeData.h b/Engine/Plugins/Developer/CustomLaunchUI/Source/Private/Marshalling/LaunchProfileTreeData.h
new file mode 100644
--- /dev/null
+++ b/Engine/Plugins/Developer/CustomLaunchUI/Source/Private/Marshalling/LaunchProfileTreeData.h
@@ -0,0 +1,106 @@
+// Copyright Epic Games, Inc. All Rights Reserved.
+
+#pragma once
+
+#include "Internationalization/Text.h"
+#include "Templates/Function.h"
+#include "Templates/SharedPointer.h"
+#include "ILauncherProfile.h"
+
+
+class SWidget;
+class ILaunchProfileMarshaller;
+class FLaunchProfileTreeData;
+class FLaunchProfileTreeNode;
+struct EVisibility;
+namespace CustomLaunchUI
+{
+	class FModel;
+};
+
+typedef TSharedPtr<FLaunchProfileTreeData> FLaunchProfileTreeDataPtr;
+typedef TSharedRef<FLaunchProfileTreeData> FLaunchProfileTreeDataRef;
+
+typedef TSharedPtr<FLaunchProfileTreeNode> FLaunchProfileTreeNodePtr;
+typedef TSharedRef<FLaunchProfileTreeNode> FLaunchProfileTreeNodeRef;
+
+
+class FLaunchProfileTreeNode : public TSharedFromThis<FLaunchProfileTreeNode>
+{
+public:
+	// construction
+	FLaunchProfileTreeNode( const FLaunchProfileTreeData* InTreeData );
+
+	typedef TFunction<void()> FFunction;
+	typedef TFunction<bool()> FGetBool;
+	typedef TFunction<void(bool)> FSetBool;
+	typedef TFunction<FString(void)> FGetString;
+	typedef TFunction<void(FString)> FSetString;
+
+	// node creation
+	struct FCallbacks
+	{
+		FGetBool IsDefault = nullptr;
+		FFunction SetToDefault = nullptr;
+		FGetBool IsVisible = nullptr;
+		FGetBool IsEnabled = nullptr;
+	};
+	FLaunchProfileTreeNode& AddWidget( FText InName, FCallbacks&& InWidgetCallbacks, TSharedRef<SWidget> InValueWidget );
+	FLaunchProfileTreeNode& AddWidget( FText InName, TSharedRef<SWidget> InValueWidget );
+
+
+	struct FBooleanCallbacks
+	{
+		FGetBool GetValue = nullptr;
+		FSetBool SetValue = nullptr;
+		FGetBool GetDefaultValue = nullptr;
+		FGetBool IsVisible = nullptr;
+		FGetBool IsEnabled = nullptr;
+	};
+	FLaunchProfileTreeNode& AddBoolean( FText InName, FBooleanCallbacks&& BooleanCallbacks );
+
+	struct FStringCallbacks
+	{
+		FGetString GetValue = nullptr;
+		FSetString SetValue = nullptr;
+		FGetString GetDefaultValue = nullptr;
+		FGetBool IsVisible = nullptr;
+		FGetBool IsEnabled = nullptr;
+
+	};
+	FLaunchProfileTreeNode& AddString( FText InName, FStringCallbacks&& StringCallbacks );
+	FLaunchProfileTreeNode& AddDirectoryString( FText InName, FStringCallbacks&& StringCallbacks );
+
+
+
+	// member variables
+	FText Name;
+	TSharedPtr<SWidget> Widget;
+
+	FCallbacks Callbacks;
+
+
+	TArray<FLaunchProfileTreeNodePtr> Children;
+	const FLaunchProfileTreeData* TreeData;
+
+protected:
+
+};
+
+
+
+class FLaunchProfileTreeData : public TSharedFromThis<FLaunchProfileTreeData>
+{
+public:
+	// construction
+	FLaunchProfileTreeData(ILauncherProfilePtr InProfile, TSharedRef<CustomLaunchUI::FModel> InModel, ILaunchProfileMarshaller* InMarshaller);
+
+	// node creation
+	FLaunchProfileTreeNode& AddHeading( FText InName );
+
+	// member variables
+	ILauncherProfilePtr Profile;
+	TSharedRef<CustomLaunchUI::FModel> Model;
+	TArray<FLaunchProfileTreeNodePtr> Nodes;
+	ILaunchProfileMarshaller* Marshaller;
+};
diff --git a/Engine/Plugins/Developer/CustomLaunchUI/Source/Private/Model/CustomLaunchUIModel.cpp b/Engine/Plugins/Developer/CustomLaunchUI/Source/Private/Model/CustomLaunchUIModel.cpp
--- a/Engine/Plugins/Developer/CustomLaunchUI/Source/Private/Model/CustomLaunchUIModel.cpp
+++ b/Engine/Plugins/Developer/CustomLaunchUI/Source/Private/Model/CustomLaunchUIModel.cpp
@@ -51,6 +51,9 @@ namespace CustomLaunchUI
 
 		AllProfiles.Add(BasicLaunchProfile);
 		SortProfiles();
+
+		DefaultBasicLaunchProfile = CreateBasicLaunchProfile();
+		DefaultCustomLaunchProfile = CreateCustomProfile(TEXT("DefaultCustomProfile"));
 	}
 
 
@@ -500,6 +503,18 @@ namespace CustomLaunchUI
 	}
 
 
+	const ILauncherProfileRef FModel::GetDefaultBasicLaunchProfile() const
+	{
+		return DefaultBasicLaunchProfile.ToSharedRef();
+	}
+
+	const ILauncherProfileRef FModel::GetDefaultCustomLaunchProfile() const
+	{
+		return DefaultCustomLaunchProfile.ToSharedRef();
+	}
+
+
+
 	TSharedPtr<FLaunchLogMessage> FModel::AddLogMessage( const FString& InMessage, ELogVerbosity::Type InVerbosity )
 	{
 		TSharedPtr<FLaunchLogMessage> Message = MakeShared<FLaunchLogMessage>(InMessage, InVerbosity);
diff --git a/Engine/Plugins/Developer/CustomLaunchUI/Source/Private/Model/CustomLaunchUIModel.h b/Engine/Plugins/Developer/CustomLaunchUI/Source/Private/Model/CustomLaunchUIModel.h
--- a/Engine/Plugins/Developer/CustomLaunchUI/Source/Private/Model/CustomLaunchUIModel.h
+++ b/Engine/Plugins/Developer/CustomLaunchUI/Source/Private/Model/CustomLaunchUIModel.h
@@ -116,6 +116,9 @@ namespace CustomLaunchUI
 		ILauncherProfileRef CreateBasicLaunchProfile();
 		ILauncherProfilePtr CloneCustomProfile(const ILauncherProfileRef& Profile);
 
+		const ILauncherProfileRef GetDefaultBasicLaunchProfile() const;
+		const ILauncherProfileRef GetDefaultCustomLaunchProfile() const;
+
 		EContentScheme DetermineProfileContentScheme(const ILauncherProfileRef& Profile) const;
 		void SetProfileContentScheme(EContentScheme ContentScheme, const ILauncherProfileRef& Profile, ELauncherProfileDeploymentModes::Type DefaultDeploymentMode = ELauncherProfileDeploymentModes::CopyToDevice);
 
@@ -160,5 +163,8 @@ namespace CustomLaunchUI
 
 		FOnSelectedProfileChanged ProfileSelectedDelegate;
 		FOnSelectedPlatformChanged PlatformChangedDelegate;
+
+		ILauncherProfilePtr DefaultBasicLaunchProfile;
+		ILauncherProfilePtr DefaultCustomLaunchProfile;
 	};
 }
diff --git a/Engine/Plugins/Developer/CustomLaunchUI/Source/Private/Widgets/Profiles/SCustomLaunchCustomProfileEditor.cpp b/Engine/Plugins/Developer/CustomLaunchUI/Source/Private/Widgets/Profiles/SCustomLaunchCustomProfileEditor.cpp
new file mode 100644
--- /dev/null
+++ b/Engine/Plugins/Developer/CustomLaunchUI/Source/Private/Widgets/Profiles/SCustomLaunchCustomProfileEditor.cpp
@@ -0,0 +1,350 @@
+// Copyright Epic Games, Inc. All Rights Reserved.
+
+#include "Widgets/Profiles/SCustomLaunchCustomProfileEditor.h"
+
+#include "Styling/AppStyle.h"
+#include "Styling/StyleColors.h"
+#include "Styling/CustomLaunchUIStyle.h"
+#include "Framework/Commands/UIAction.h"
+#include "Framework/MultiBox/MultiBoxBuilder.h"
+#include "SlateOptMacros.h"
+#include "Widgets/SBoxPanel.h"
+#include "Widgets/Input/SHyperlink.h"
+#include "Widgets/Input/SComboBox.h"
+#include "Widgets/Input/SComboButton.h"
+#include "Widgets/Layout/SSplitter.h"
+#include "Widgets/Layout/SScrollBorder.h"
+#include "Widgets/Text/STextBlock.h"
+#include "Widgets/Text/SInlineEditableTextBlock.h"
+#include "Widgets/Images/SImage.h"
+#include "Widgets/Views/STreeView.h"
+#include "Widgets/SNullWidget.h"
+#include "Widgets/Shared/SCustomLaunchDeviceCombo.h"
+#include "Widgets/Shared/SCustomLaunchDeviceListView.h"
+#include "SPositiveActionButton.h"
+#include "DesktopPlatformModule.h"
+#include "PlatformInfo.h"
+#include "ITargetDeviceProxy.h"
+#include "ITargetDeviceProxyManager.h"
+#include "Interfaces/ITargetPlatform.h"
+#include "Interfaces/ITargetPlatformManagerModule.h"
+#include "InstalledPlatformInfo.h"
+#include "Framework/Docking/TabManager.h"
+#include "Misc/ConfigCacheIni.h"
+#include "GameProjectHelper.h"
+
+
+#define LOCTEXT_NAMESPACE "SCustomLaunchCustomProfileEditor"
+
+
+
+
+
+class SLaunchProfileCategoryTreeRow : public SLaunchProfileTreeRow
+{
+public:
+	SLATE_BEGIN_ARGS(SLaunchProfileCategoryTreeRow){}
+	SLATE_END_ARGS()
+
+	void Construct(const FArguments& InArgs, FLaunchProfileTreeNodePtr TreeNode, const TSharedRef<STableViewBase>& OwnerTable )
+	{
+		this->ChildSlot
+		[
+			SNew(SBorder)
+			.BorderImage(FAppStyle::GetBrush("Brushes.Background"))
+			.Padding(FMargin(0, 0, 0, 1))
+			[
+				SNew(SBorder)
+				.BorderImage(FAppStyle::GetBrush("Brushes.Header"))				
+				.Padding(0)
+				[
+					SNew(SBox)
+					.MinDesiredHeight(26.0f)
+					[
+						SNew(SHorizontalBox)
+
+						+SHorizontalBox::Slot()
+						.HAlign(HAlign_Left)
+						.VAlign(VAlign_Center)
+						.Padding(2, 0, 0, 0)
+						.AutoWidth()
+						[
+							SNew(SExpanderArrow, SharedThis(this))
+							.StyleSet( &FCoreStyle::Get() )
+						]
+
+						+SHorizontalBox::Slot()
+						.VAlign(VAlign_Center)
+						.Padding(4, 0, 0, 0)
+						.FillWidth(1)
+						[
+							SNew(STextBlock)
+							.Text(TreeNode->Name)
+							.Font(FCoreStyle::Get().GetFontStyle("SmallFontBold"))
+						]
+					]
+				]
+			]
+		];
+
+		SLaunchProfileTreeRow::ConstructInternal(
+			STableRow::FArguments()
+			.ShowSelection(false),
+			OwnerTable
+		);
+	}
+};
+
+class SLaunchProfilePropertyTreeRow : public SLaunchProfileTreeRow
+{
+public:
+	DECLARE_DELEGATE_OneParam( FOnSplitterResized, float );
+
+
+	SLATE_BEGIN_ARGS(SLaunchProfilePropertyTreeRow){}
+		SLATE_ATTRIBUTE(float, SplitterValue)
+		SLATE_EVENT(FOnSplitterResized, OnSplitterResized)
+	SLATE_END_ARGS()
+
+	BEGIN_SLATE_FUNCTION_BUILD_OPTIMIZATION
+	void Construct(const FArguments& InArgs, FLaunchProfileTreeNodePtr TreeNode, const TSharedRef<STableViewBase>& OwnerTable )
+	{
+		SplitterValue = InArgs._SplitterValue;
+		OnSplitterResized = InArgs._OnSplitterResized;
+
+
+		TSharedRef<SWidget> OptionalResetToDefaultWidget = SNullWidget::NullWidget;
+		if (TreeNode->Callbacks.IsDefault != nullptr && TreeNode->Callbacks.SetToDefault != nullptr)
+		{
+			auto IsVisible = [TreeNode]()
+			{
+				if (TreeNode->Callbacks.IsDefault())
+				{
+					return false;
+				}
+
+				if (TreeNode->Callbacks.IsEnabled && !TreeNode->Callbacks.IsEnabled())
+				{
+					return false;
+				}
+
+				return true;
+			};
+
+			auto SetToDefault = [TreeNode]()
+			{
+				TreeNode->Callbacks.SetToDefault();
+				TreeNode->TreeData->Marshaller->OnPropertyChanged();
+				return FReply::Handled();
+			};
+
+			OptionalResetToDefaultWidget = SNew(SButton)
+				.ButtonStyle(FCustomLaunchUIStyle::Get(), "HoverHintOnly")
+				.ToolTipText( LOCTEXT("ResetToDefaultToolTip", "Reset this property to its default value.") )
+				.Visibility_Lambda( [IsVisible]() { return IsVisible() ? EVisibility::Visible : EVisibility::Collapsed; } )
+				.OnClicked_Lambda(SetToDefault)
+				.ContentPadding(0)
+				[
+					SNew(SImage)
+					.Image( FCustomLaunchUIStyle::Get().GetBrush("Icons.DiffersFromDefault") )
+					.DesiredSizeOverride(FVector2D(16,16))
+					.ColorAndOpacity( FSlateColor::UseForeground() )
+				]
+			;
+		}
+
+		auto MakeSplitterSlot = [this]( TSharedRef<SWidget> SlotContent )
+		{
+			return SNew(SBorder)
+			.BorderImage(FAppStyle::Get().GetBrush("Brushes.Recessed"))
+			.Padding(FMargin(0, 0, 0, 1))
+			[
+				SNew(SHorizontalBox)
+				+SHorizontalBox::Slot()
+				.HAlign(HAlign_Fill)
+				.VAlign(VAlign_Fill)
+				[
+					SNew( SBorder )
+					.BorderImage_Lambda( [this]() { return IsHovered() ? FAppStyle::Get().GetBrush("Brushes.Header") : FAppStyle::Get().GetBrush("Brushes.Panel"); } )
+					.Padding(0)
+					[
+						SlotContent
+					]
+				]
+			];
+		};
+
+
+		this->ChildSlot
+		[
+			SNew(SSplitter)
+			.Orientation(Orient_Horizontal)
+			.PhysicalSplitterHandleSize(1.0)
+			.HitDetectionSplitterHandleSize(5.0)
+			.MinimumSlotHeight(26.0f)
+
+			// property name
+			+SSplitter::Slot()
+			.Value_Lambda( [this]() { return 1.0f - SplitterValue.Get(); } )
+			.OnSlotResized_Lambda( [this](float NewPos) { OnSplitterResized.ExecuteIfBound(1.0f - FMath::Clamp(NewPos,0.0f, 1.0f)); } )
+			.MinSize(8)
+			[
+				MakeSplitterSlot(
+					SNew(SBox)
+					.Padding( FMargin(24, 4, 4, 4))
+					.VAlign(VAlign_Top)
+					[
+						SNew(STextBlock)
+						.Text(TreeNode->Name)
+						.Font(FCoreStyle::Get().GetFontStyle("SmallFont"))
+					]
+				)
+			]
+
+			// property value
+			+SSplitter::Slot()
+			.Value_Lambda( [this](){ return SplitterValue.Get(); } )
+			.OnSlotResized_Lambda( [this](float NewPos) { OnSplitterResized.ExecuteIfBound(FMath::Clamp(NewPos,0.0f, 1.0f)); } )
+			.MinSize(32)
+			[
+				MakeSplitterSlot(
+					SNew(SBox)
+					.Padding( FMargin(8, 1, 1, 1))
+					.VAlign(VAlign_Center)
+					[
+						TreeNode->Widget.ToSharedRef()
+					]
+				)
+			]
+
+			// reset to default button
+			+SSplitter::Slot()
+			.MinSize(24)
+			.Resizable(false)
+			.Value(0)
+			[
+				MakeSplitterSlot(
+					SNew(SBox)
+					.Padding(0)
+					.HAlign(HAlign_Center)
+					.VAlign(VAlign_Center)
+					[
+						OptionalResetToDefaultWidget
+					]
+				)
+			]
+		];
+
+
+		SLaunchProfileTreeRow::ConstructInternal(
+			STableRow::FArguments()
+			.ShowSelection(false),
+			OwnerTable
+		);
+	}
+	END_SLATE_FUNCTION_BUILD_OPTIMIZATION
+
+
+	FOnSplitterResized OnSplitterResized;
+	TAttribute<float> SplitterValue;
+};
+
+
+
+BEGIN_SLATE_FUNCTION_BUILD_OPTIMIZATION
+void SCustomLaunchCustomProfileEditor::Construct(const FArguments& InArgs, const TSharedRef<CustomLaunchUI::FModel>& InModel)
+{
+	Model = InModel;
+
+	Marshaller = ILaunchProfileMarshallerFactory::Marshall( nullptr, InModel );
+
+
+	ChildSlot
+	[
+		SAssignNew(TreeView, SLaunchProfileTreeView)
+		.TreeItemsSource(&Marshaller->GetProfileTree()->Nodes)
+		.SelectionMode( ESelectionMode::None )
+		.OnGenerateRow( this, &SCustomLaunchCustomProfileEditor::OnGenerateWidgetForTreeNode )
+		.OnGetChildren( this, &SCustomLaunchCustomProfileEditor::OnGetChildren)
+		.HandleDirectionalNavigation(false)
+	];
+
+}
+END_SLATE_FUNCTION_BUILD_OPTIMIZATION
+
+TSharedRef<ITableRow> SCustomLaunchCustomProfileEditor::OnGenerateWidgetForTreeNode( FLaunchProfileTreeNodePtr TreeNode, const TSharedRef<STableViewBase>& OwnerTable )
+{
+	if (TreeNode->Widget.IsValid())
+	{
+		auto IsVisible = [TreeNode]()
+		{
+			if (TreeNode->Callbacks.IsVisible == nullptr)
+			{
+				return true;
+			}
+
+			return TreeNode->Callbacks.IsVisible();
+		};
+
+		return SNew(SLaunchProfilePropertyTreeRow, TreeNode, OwnerTable)
+		.SplitterValue_Lambda( [this]() { return SplitterPos; } )
+		.OnSplitterResized_Lambda( [this](float NewPos) { SplitterPos = FMath::Clamp(NewPos, 0.0f, 1.0f); } )
+		.Visibility_Lambda( [IsVisible]() { return IsVisible() ? EVisibility::Visible : EVisibility::Collapsed; } )
+		.IsEnabled_Lambda( [TreeNode]() { return TreeNode->Callbacks.IsEnabled ? TreeNode->Callbacks.IsEnabled() : true; } )
+		;
+	}
+	else
+	{
+		return SNew(SLaunchProfileCategoryTreeRow, TreeNode, OwnerTable);
+	}
+}
+
+void SCustomLaunchCustomProfileEditor::OnGetChildren(FLaunchProfileTreeNodePtr TreeNode, TArray<FLaunchProfileTreeNodePtr>& OutChildren)
+{
+	OutChildren = TreeNode->Children;
+}
+
+FVector2D SCustomLaunchCustomProfileEditor::GetScrollDistance()
+{
+	if (!TreeView.IsValid() || !Marshaller.IsValid())
+	{
+		return FVector2D::ZeroVector;
+	}
+
+	return TreeView->GetScrollDistance();
+}
+
+FVector2D SCustomLaunchCustomProfileEditor::GetScrollDistanceRemaining()
+{
+	if (!TreeView.IsValid() || !Marshaller.IsValid())
+	{
+		return FVector2D::ZeroVector;
+	}
+
+	return TreeView->GetScrollDistanceRemaining();
+}
+
+TSharedRef<SWidget> SCustomLaunchCustomProfileEditor::GetScrollWidget()
+{
+	return SharedThis(this);
+}
+
+
+void SCustomLaunchCustomProfileEditor::SetProfile( const ILauncherProfilePtr& Profile )
+{
+	CurrentProfile = Profile;
+
+
+	Marshaller = ILaunchProfileMarshallerFactory::Marshall( Profile, Model.ToSharedRef() );
+
+	TreeView->SetTreeItemsSource( &Marshaller->GetProfileTree()->Nodes );
+	for ( FLaunchProfileTreeNodePtr Node : TreeView->GetRootItems())
+	{
+		TreeView->SetItemExpansion(Node, true);
+	}
+
+	TreeView->RequestTreeRefresh();
+}
+
+
+#undef LOCTEXT_NAMESPACE
diff --git a/Engine/Plugins/Developer/CustomLaunchUI/Source/Private/Widgets/Profiles/SCustomLaunchCustomProfileEditor.h b/Engine/Plugins/Developer/CustomLaunchUI/Source/Private/Widgets/Profiles/SCustomLaunchCustomProfileEditor.h
new file mode 100644
--- /dev/null
+++ b/Engine/Plugins/Developer/CustomLaunchUI/Source/Private/Widgets/Profiles/SCustomLaunchCustomProfileEditor.h
@@ -0,0 +1,56 @@
+// Copyright Epic Games, Inc. All Rights Reserved.
+
+#pragma once
+
+#include "Internationalization/Text.h"
+#include "Framework/SlateDelegates.h"
+#include "Framework/Layout/IScrollableWidget.h"
+#include "Widgets/SCompoundWidget.h"
+#include "Model/CustomLaunchUIModel.h"
+#include "Widgets/Views/STableRow.h"
+#include "Marshalling/ILaunchProfileMarshaller.h"
+
+
+class SWidget;
+class ITableRow;
+class STableViewBase;
+template<typename T> class STreeView;
+
+typedef STreeView<FLaunchProfileTreeNodePtr> SLaunchProfileTreeView;
+typedef STableRow<FLaunchProfileTreeNodePtr> SLaunchProfileTreeRow;
+
+
+class SCustomLaunchCustomProfileEditor : public SCompoundWidget, public IScrollableWidget
+{
+public:
+	SLATE_BEGIN_ARGS(SCustomLaunchCustomProfileEditor) {}
+	SLATE_END_ARGS()
+
+public:
+	void Construct(const FArguments& InArgs, const TSharedRef<CustomLaunchUI::FModel>& InModel);
+
+	void SetProfile( const ILauncherProfilePtr& Profile );
+	ILauncherProfilePtr GetProfile() const { return CurrentProfile; }
+
+
+	// IScrollableWidget interface
+	virtual FVector2D GetScrollDistance() override;
+	virtual FVector2D GetScrollDistanceRemaining() override;
+	virtual TSharedRef<SWidget> GetScrollWidget() override;
+
+private:
+
+	TSharedRef<ITableRow> OnGenerateWidgetForTreeNode( FLaunchProfileTreeNodePtr TreeNode, const TSharedRef<STableViewBase>& OwnerTable );
+	void OnGetChildren(FLaunchProfileTreeNodePtr TreeNode, TArray<FLaunchProfileTreeNodePtr>& OutChildren);
+
+	ILauncherProfilePtr CurrentProfile;
+	TSharedPtr<CustomLaunchUI::FModel> Model;
+	TSharedPtr<SLaunchProfileTreeView> TreeView;
+	TSharedPtr<ILaunchProfileMarshaller> Marshaller;
+
+	float SplitterPos = 0.6f;
+	
+
+private:
+
+};
diff --git a/Engine/Plugins/Developer/CustomLaunchUI/Source/Private/Widgets/Profiles/SCustomLaunchProfilesPanel.cpp b/Engine/Plugins/Developer/CustomLaunchUI/Source/Private/Widgets/Profiles/SCustomLaunchProfilesPanel.cpp
--- a/Engine/Plugins/Developer/CustomLaunchUI/Source/Private/Widgets/Profiles/SCustomLaunchProfilesPanel.cpp
+++ b/Engine/Plugins/Developer/CustomLaunchUI/Source/Private/Widgets/Profiles/SCustomLaunchProfilesPanel.cpp
@@ -17,15 +17,13 @@
 #include "Widgets/Shared/SCustomLaunchProjectCombo.h"
 #include "Widgets/Shared/SCustomLaunchBuildTargetCombo.h"
 #include "Widgets/Profiles/SCustomLaunchCustomProfileSelector.h"
+#include "Widgets/Profiles/SCustomLaunchCustomProfileEditor.h"
 #include "Widgets/Output/SCustomLaunchLaunchPanel.h"
 #include "Widgets/Output/SCustomLaunchOutputLog.h"
-#include "Classes/LauncherProfileSettings.h"
 #include "ILauncher.h"
 #include "ITargetDeviceProxyManager.h"
 #include "ITargetDeviceProxy.h"
-#include "IDetailsView.h"
 #include "Modules/ModuleManager.h"
-#include "PropertyEditorModule.h"
 #include "Framework/Docking/TabManager.h"
 #include "PlatformInfo.h"
 
@@ -40,11 +38,6 @@ SCustomLaunchProfilesPanel::SCustomLaunchProfilesPanel()
 
 SCustomLaunchProfilesPanel::~SCustomLaunchProfilesPanel()
 {
-	if (CurrentProfileSettingsObject != nullptr)
-	{
-		CurrentProfileSettingsObject->RemoveFromRoot();
-	}
-
 	if (Model.IsValid())
 	{
 		Model->OnProfileSelected().RemoveAll(this);
@@ -173,6 +166,7 @@ TSharedRef<SWidget> SCustomLaunchProfilesPanel::CreateProfilesPanel()
 			[
 				SNew(STextBlock)
 				.Text(LOCTEXT("DefaultBuildTargetSelectLabel", "Target"))
+				.Visibility(CustomLaunchUI::bUseFriendlyBuildTargetSelection ? EVisibility::Collapsed : EVisibility::Visible)
 			]
 
 			// Default Build Target selector
@@ -185,6 +179,7 @@ TSharedRef<SWidget> SCustomLaunchProfilesPanel::CreateProfilesPanel()
 				.OnSelectionChanged(this, &SCustomLaunchProfilesPanel::SetDefaultBuildTarget)
 				.SelectedBuildTarget(this, &SCustomLaunchProfilesPanel::GetDefaultBuildTarget)
 				.SelectedProject(this, &SCustomLaunchProfilesPanel::GetDefaultProjectPath)
+				.Visibility(CustomLaunchUI::bUseFriendlyBuildTargetSelection ? EVisibility::Collapsed : EVisibility::Visible)
 			]
 		];		
 	}
@@ -314,19 +309,7 @@ END_SLATE_FUNCTION_BUILD_OPTIMIZATION
 BEGIN_SLATE_FUNCTION_BUILD_OPTIMIZATION
 TSharedRef<SWidget> SCustomLaunchProfilesPanel::CreateProfileEditorWidget()
 {
-	FDetailsViewArgs DetailsViewArgs;
-	DetailsViewArgs.NameAreaSettings = FDetailsViewArgs::HideNameArea;
-	DetailsViewArgs.bAllowSearch = false;
-	DetailsViewArgs.bHideSelectionTip = true;
-	DetailsViewArgs.bShowPropertyMatrixButton = false;
-	DetailsViewArgs.bShowModifiedPropertiesOption = false;
-	DetailsViewArgs.bShowHiddenPropertiesWhilePlayingOption = false;
-	DetailsViewArgs.bAllowMultipleTopLevelObjects = false;
-	DetailsViewArgs.bShowKeyablePropertiesOption = false;
-	DetailsViewArgs.bShowAnimatedPropertiesOption = false;
-
-	FPropertyEditorModule& PropertyEditorModule = FModuleManager::GetModuleChecked<FPropertyEditorModule>("PropertyEditor");
-	DetailsView = PropertyEditorModule.CreateDetailView(DetailsViewArgs);
+	PropertyEditor = SNew(SCustomLaunchCustomProfileEditor, Model.ToSharedRef());
 
 	return SNew(SVerticalBox)
 	.Visibility_Lambda( [this] { return Model->GetSelectedProfile().IsValid() && IsProfileEditorVisible() ? EVisibility::Visible :EVisibility::Collapsed; } )
@@ -342,9 +325,9 @@ TSharedRef<SWidget> SCustomLaunchProfilesPanel::CreateProfileEditorWidget()
 	+SVerticalBox::Slot()
 	.FillHeight(1)
 	[
-		SNew(SScrollBorder, DetailsView.ToSharedRef())
+		SNew(SScrollBorder, PropertyEditor.ToSharedRef())
 		[
-			DetailsView.ToSharedRef()
+			PropertyEditor.ToSharedRef()
 		]
 	]
 
@@ -670,43 +653,21 @@ FReply SCustomLaunchProfilesPanel::OnOpenDeviceManagerClicked()
 
 void SCustomLaunchProfilesPanel::OnProfileSelected( const ILauncherProfilePtr& NewProfile, const ILauncherProfilePtr& OldProfile)
 {
-	DetailsView->SetObject(nullptr);
-	if (CurrentProfileSettingsObject != nullptr)
+	if (NewProfile.IsValid() && !Model->IsAdvancedProfile(NewProfile.ToSharedRef()))
 	{
-		CurrentProfileSettingsObject->RemoveFromRoot();
-		CurrentProfileSettingsObject = nullptr;
+		PropertyEditor->SetProfile(NewProfile);
 	}
-
-	if (NewProfile.IsValid() && !Model->IsAdvancedProfile(NewProfile.ToSharedRef()))
+	else
 	{
-		if (Model->IsBasicLaunchProfile(NewProfile))
-		{
-			CurrentProfileSettingsObject = NewObject<ULauncherBasicLaunchProfileSettings>();
-		}
-		else
-		{
-			CurrentProfileSettingsObject = NewObject<ULauncherCustomProfileSettings>();
-		}
-
-		CurrentProfileSettingsObject->AddToRoot();
-		CurrentProfileSettingsObject->SetProfile(Model, NewProfile);
-		DetailsView->SetObject(CurrentProfileSettingsObject);
+		PropertyEditor->SetProfile(nullptr);
 	}
 }
 
 void SCustomLaunchProfilesPanel::OnProfileModified(const ILauncherProfilePtr& Profile)
 {
-	if (Model->GetSelectedProfile() == Profile && Profile.IsValid())
+	if (!Model->IsBasicLaunchProfile(Profile.ToSharedRef()))
 	{
-		TArray<TWeakObjectPtr<UObject>> SelectedObjects = DetailsView->GetSelectedObjects();
-		if (SelectedObjects.Num() == 1)
-		{
-			UBaseLauncherProfileSettings* ProfileSettings = Cast<UBaseLauncherProfileSettings>(SelectedObjects[0].Get());
-			if (ProfileSettings != nullptr)
-			{
-				ProfileSettings->RefreshProfile();
-			}
-		}
+		Model->GetProfileManager()->SaveJSONProfile(Profile.ToSharedRef());
 	}
 }
 
@@ -825,6 +786,12 @@ bool SCustomLaunchProfilesPanel::IsProfileEditorVisible() const
 void SCustomLaunchProfilesPanel::SetDefaultProjectPath(FString ProjectPath)
 {
 	Model->GetProfileManager()->SetProjectPath(ProjectPath);
+	Model->GetProfileManager()->SetBuildTarget(FString());
+
+	if (Model->GetSelectedProfile().IsValid())
+	{
+		Model->UpdateCookedPlatformsFromBuildTarget(Model->GetSelectedProfile().ToSharedRef());
+	}
 }
 
 FString SCustomLaunchProfilesPanel::GetDefaultProjectPath() const
@@ -835,6 +802,11 @@ FString SCustomLaunchProfilesPanel::GetDefaultProjectPath() const
 void SCustomLaunchProfilesPanel::SetDefaultBuildTarget(FString BuildTarget)
 {
 	Model->GetProfileManager()->SetBuildTarget(BuildTarget);
+
+	if (Model->GetSelectedProfile().IsValid())
+	{
+		Model->UpdateCookedPlatformsFromBuildTarget(Model->GetSelectedProfile().ToSharedRef());
+	}
 }
 
 FString SCustomLaunchProfilesPanel::GetDefaultBuildTarget() const
diff --git a/Engine/Plugins/Developer/CustomLaunchUI/Source/Private/Widgets/Profiles/SCustomLaunchProfilesPanel.h b/Engine/Plugins/Developer/CustomLaunchUI/Source/Private/Widgets/Profiles/SCustomLaunchProfilesPanel.h
--- a/Engine/Plugins/Developer/CustomLaunchUI/Source/Private/Widgets/Profiles/SCustomLaunchProfilesPanel.h
+++ b/Engine/Plugins/Developer/CustomLaunchUI/Source/Private/Widgets/Profiles/SCustomLaunchProfilesPanel.h
@@ -82,8 +82,7 @@ private:
 	TSharedPtr<SExpandableArea> LogExpandableArea;
 	TSharedPtr<SSplitter> ProfileEditorSplitter;
 	TSharedPtr<SCustomLaunchCustomProfileSelector> ProfileSelector;
-	TSharedPtr<class IDetailsView> DetailsView;
-	TObjectPtr<class UBaseLauncherProfileSettings> CurrentProfileSettingsObject;
+	TSharedPtr<class SCustomLaunchCustomProfileEditor> PropertyEditor;
 	TSharedPtr<class SCustomLaunchOutputLog> OutputLog;
 
 
diff --git a/Engine/Plugins/Developer/CustomLaunchUI/Source/Private/Widgets/Shared/SCustomLaunchBuildConfigurationCombo.cpp b/Engine/Plugins/Developer/CustomLaunchUI/Source/Private/Widgets/Shared/SCustomLaunchBuildConfigurationCombo.cpp
new file mode 100644
--- /dev/null
+++ b/Engine/Plugins/Developer/CustomLaunchUI/Source/Private/Widgets/Shared/SCustomLaunchBuildConfigurationCombo.cpp
@@ -0,0 +1,83 @@
+// Copyright Epic Games, Inc. All Rights Reserved.
+
+#include "Widgets/Shared/SCustomLaunchBuildConfigurationCombo.h"
+
+#include "SlateOptMacros.h"
+#include "Widgets/Input/SComboButton.h"
+#include "Styling/CustomLaunchUIStyle.h"
+#include "Framework/MultiBox/MultiBoxBuilder.h"
+#include "HAL/PlatformMisc.h"
+
+#define LOCTEXT_NAMESPACE "SCustomLaunchBuildConfigurationCombo"
+
+
+BEGIN_SLATE_FUNCTION_BUILD_OPTIMIZATION
+void SCustomLaunchBuildConfigurationCombo::Construct(const FArguments& InArgs)
+{
+	OnSelectionChanged = InArgs._OnSelectionChanged;
+	SelectedBuildConfiguration = InArgs._SelectedBuildConfiguration;
+
+	FSlateFontInfo Font = InArgs._Font.IsSet() ? InArgs._Font.Get() : InArgs._TextStyle->Font;
+
+	ChildSlot
+	[
+		SNew(SComboButton)
+		.ButtonContent()
+		[
+			SNew(STextBlock)
+			.Text(this, &SCustomLaunchBuildConfigurationCombo::GetBuildConfigurationName)
+			.Font(Font)
+		]
+		.OnGetMenuContent(this, &SCustomLaunchBuildConfigurationCombo::MakeBuildConfigurationSelectionWidget)
+	];
+}
+END_SLATE_FUNCTION_BUILD_OPTIMIZATION
+
+
+
+BEGIN_SLATE_FUNCTION_BUILD_OPTIMIZATION
+TSharedRef<SWidget> SCustomLaunchBuildConfigurationCombo::MakeBuildConfigurationSelectionWidget()
+{
+	FMenuBuilder MenuBuilder(true, nullptr);
+	{
+		static const EBuildConfiguration SupportedConfigurations[] =
+		{
+			EBuildConfiguration::Debug,
+			EBuildConfiguration::Development,
+			EBuildConfiguration::Test,
+			EBuildConfiguration::Shipping,
+		};
+
+		for (const EBuildConfiguration& Configuration : SupportedConfigurations)
+		{
+			MenuBuilder.AddMenuEntry(
+				EBuildConfigurations::ToText(Configuration), 
+				LOCTEXT("SelectConfigTip","Select this configuration"), 
+				FSlateIcon(), 
+				FUIAction(
+					FExecuteAction::CreateSP(this, &SCustomLaunchBuildConfigurationCombo::SetBuildConfiguration, Configuration),
+					FCanExecuteAction(),
+					FGetActionCheckState::CreateLambda( [this, Configuration]() { return (SelectedBuildConfiguration.Get() == Configuration) ? ECheckBoxState::Checked : ECheckBoxState::Unchecked; })
+				),
+				NAME_None,
+				EUserInterfaceActionType::Check);
+		}
+	}
+
+	return MenuBuilder.MakeWidget();
+}
+END_SLATE_FUNCTION_BUILD_OPTIMIZATION
+
+
+FText SCustomLaunchBuildConfigurationCombo::GetBuildConfigurationName() const
+{
+	EBuildConfiguration BuildConfiguration = SelectedBuildConfiguration.Get();
+	return EBuildConfigurations::ToText(BuildConfiguration);
+}
+
+void SCustomLaunchBuildConfigurationCombo::SetBuildConfiguration(EBuildConfiguration BuildConfiguration)
+{
+	OnSelectionChanged.ExecuteIfBound(BuildConfiguration);
+}
+
+#undef LOCTEXT_NAMESPACE
diff --git a/Engine/Plugins/Developer/CustomLaunchUI/Source/Private/Widgets/Shared/SCustomLaunchBuildConfigurationCombo.h b/Engine/Plugins/Developer/CustomLaunchUI/Source/Private/Widgets/Shared/SCustomLaunchBuildConfigurationCombo.h
new file mode 100644
--- /dev/null
+++ b/Engine/Plugins/Developer/CustomLaunchUI/Source/Private/Widgets/Shared/SCustomLaunchBuildConfigurationCombo.h
@@ -0,0 +1,41 @@
+// Copyright Epic Games, Inc. All Rights Reserved.
+
+#pragma once
+
+#include "CoreMinimal.h"
+#include "Widgets/SCompoundWidget.h"
+#include "Styling/SlateTypes.h"
+#include "Styling/CoreStyle.h"
+#include "Styling/SlateWidgetStyleAsset.h"
+#include "Fonts/SlateFontInfo.h"
+
+class SCustomLaunchBuildConfigurationCombo
+	: public SCompoundWidget
+{
+public:
+	DECLARE_DELEGATE_OneParam(FOnSelectionChanged, EBuildConfiguration );
+
+	SLATE_BEGIN_ARGS(SCustomLaunchBuildConfigurationCombo)
+		: _TextStyle( &FCoreStyle::Get().GetWidgetStyle<FTextBlockStyle>( "NormalText" ) )
+		{}
+		SLATE_EVENT(FOnSelectionChanged, OnSelectionChanged);
+		SLATE_ATTRIBUTE(EBuildConfiguration, SelectedBuildConfiguration);
+		SLATE_ATTRIBUTE(FSlateFontInfo, Font)
+		SLATE_STYLE_ARGUMENT( FTextBlockStyle, TextStyle )
+	SLATE_END_ARGS()
+
+public:
+	void Construct(	const FArguments& InArgs);
+
+protected:
+	TAttribute<EBuildConfiguration> SelectedBuildConfiguration;
+	FOnSelectionChanged OnSelectionChanged;
+	bool bShowAnyProjectOption;
+
+	TSharedRef<SWidget> MakeBuildConfigurationSelectionWidget();
+
+	FText GetBuildConfigurationName() const;
+	void SetBuildConfiguration(EBuildConfiguration BuildConfigurationName);
+
+
+};
diff --git a/Engine/Plugins/Developer/CustomLaunchUI/Source/Private/Widgets/Shared/SCustomLaunchProjectCombo.cpp b/Engine/Plugins/Developer/CustomLaunchUI/Source/Private/Widgets/Shared/SCustomLaunchProjectCombo.cpp
--- a/Engine/Plugins/Developer/CustomLaunchUI/Source/Private/Widgets/Shared/SCustomLaunchProjectCombo.cpp
+++ b/Engine/Plugins/Developer/CustomLaunchUI/Source/Private/Widgets/Shared/SCustomLaunchProjectCombo.cpp
@@ -6,7 +6,6 @@
 #include "Widgets/Input/SComboButton.h"
 #include "Framework/Application/SlateApplication.h"
 #include "Framework/MultiBox/MultiBoxBuilder.h"
-#include "Interfaces/IMainFrameModule.h"
 #include "Model/CustomLaunchUIModel.h"
 #include "GameProjectHelper.h"
 #include "Misc/Paths.h"
@@ -72,17 +71,17 @@ TSharedRef<SWidget> SCustomLaunchProjectCombo::MakeProjectSelectionWidget()
 				}
 			}
 
-			if (ProjectName.IsEmpty())
+			if (ProjectName.IsEmpty() && GIsEditor)
 			{
 				FUIAction CurrentProjectAction(FExecuteAction::CreateSP(this, &SCustomLaunchProjectCombo::SetProjectPath, FString()));
 				MenuBuilder.AddMenuEntry(LOCTEXT("CurrentProjectAction", "Current Project"), LOCTEXT("CurrentProjectHint", "Use the current project"), FSlateIcon(), CurrentProjectAction);
 			}
-			else if (CurrentProjectOption == ECurrentProjectOption::Empty)
+			else if (CurrentProjectOption == ECurrentProjectOption::Empty && !ProjectName.IsEmpty() )
 			{
 				FUIAction CurrentProjectAction(FExecuteAction::CreateSP(this, &SCustomLaunchProjectCombo::SetProjectPath, FString()));
 				MenuBuilder.AddMenuEntry(FText::Format(LOCTEXT("CurrentProjectActionFmt", "Current Project ({0})"), FText::FromString(ProjectName)), LOCTEXT("CurrentProjectHint", "Use the current project"), FSlateIcon(), CurrentProjectAction);
 			}
-			else
+			else if (!ProjectName.IsEmpty())
 			{
 				FUIAction CurrentProjectAction(FExecuteAction::CreateSP(this, &SCustomLaunchProjectCombo::SetProjectPath, ProjectPath));
 				MenuBuilder.AddMenuEntry(FText::FromString(*ProjectName), FText::FromString(ProjectPath), FSlateIcon(), CurrentProjectAction);
@@ -131,7 +130,7 @@ FText SCustomLaunchProjectCombo::GetProjectName() const
 {
 	FString ProjectPath = SelectedProject.Get();
 
-	if (!ProjectPath.IsEmpty() && HasProject.Get())
+	if (!ProjectPath.IsEmpty() && (HasProject.Get() || CurrentProjectOption == ECurrentProjectOption::None))
 	{
 		return FText::FromString(FPaths::GetBaseFilename(ProjectPath));
 	}
