fix for horizontal scroll bar not working when scrolled to the bottom of the log

#rnx

diff --git a/Engine/Plugins/Developer/CustomLaunchUI/Source/Private/Widgets/Output/CustomLaunchOutputLogMarshaller.cpp b/Engine/Plugins/Developer/CustomLaunchUI/Source/Private/Widgets/Output/CustomLaunchOutputLogMarshaller.cpp
--- a/Engine/Plugins/Developer/CustomLaunchUI/Source/Private/Widgets/Output/CustomLaunchOutputLogMarshaller.cpp
+++ b/Engine/Plugins/Developer/CustomLaunchUI/Source/Private/Widgets/Output/CustomLaunchOutputLogMarshaller.cpp
@@ -77,7 +77,7 @@ namespace CustomLaunchUI
 	}
 
 
-	void FLaunchLogTextLayoutMarshaller::FlushPendingLogMessages()
+	bool FLaunchLogTextLayoutMarshaller::FlushPendingLogMessages()
 	{
 		bool bRefreshLog = false;
 
@@ -109,6 +109,8 @@ namespace CustomLaunchUI
 		{
 			TextLayout->AddLines(LinesToAdd);
 		}
+
+		return bRefreshLog;
 	}
 
 
diff --git a/Engine/Plugins/Developer/CustomLaunchUI/Source/Private/Widgets/Output/CustomLaunchOutputLogMarshaller.h b/Engine/Plugins/Developer/CustomLaunchUI/Source/Private/Widgets/Output/CustomLaunchOutputLogMarshaller.h
--- a/Engine/Plugins/Developer/CustomLaunchUI/Source/Private/Widgets/Output/CustomLaunchOutputLogMarshaller.h
+++ b/Engine/Plugins/Developer/CustomLaunchUI/Source/Private/Widgets/Output/CustomLaunchOutputLogMarshaller.h
@@ -40,7 +40,7 @@ namespace CustomLaunchUI
 		void SetFilterString( const FString& InLogFilterString );
 
 		void AddPendingLogMessage( TSharedPtr<FLaunchLogMessage> Message );
-		void FlushPendingLogMessages();
+		bool FlushPendingLogMessages();
 
 		void RefreshAllLogMessages();
 		int32 GetNumFilteredMessages() const { return NumFilteredMessages; }
diff --git a/Engine/Plugins/Developer/CustomLaunchUI/Source/Private/Widgets/Output/SCustomLaunchLaunchPanel.cpp b/Engine/Plugins/Developer/CustomLaunchUI/Source/Private/Widgets/Output/SCustomLaunchLaunchPanel.cpp
--- a/Engine/Plugins/Developer/CustomLaunchUI/Source/Private/Widgets/Output/SCustomLaunchLaunchPanel.cpp
+++ b/Engine/Plugins/Developer/CustomLaunchUI/Source/Private/Widgets/Output/SCustomLaunchLaunchPanel.cpp
@@ -368,8 +368,10 @@ void SCustomLaunchLaunchPanel::SetLauncherWorker( const ILauncherWorkerRef& Work
 
 void SCustomLaunchLaunchPanel::Tick( const FGeometry& AllottedGeometry, const double InCurrentTime, const float InDeltaTime )
 {
-	LaunchLogTextMarshaller->FlushPendingLogMessages();
-	OutputLog->RequestForceScroll(true);
+	if (LaunchLogTextMarshaller->FlushPendingLogMessages())
+	{
+		OutputLog->RequestForceScroll(true);
+	}
 }
 
 
