CustomLaunchUI changes
- update Launcher Worker's UAT command line to reflect new zen pak streaming UAT parameters.
- helper function to query project settings.
- disable the Loose files option if the project is using Zen Store via Project Settings because there's no way to opt-out of Zen Store from the command line.
- disable  zen pak streaming option if automatic workspace configuration has not been set.
- disable zen streaming options to remote devices if AllowRemoteNetworkService is disabled.
- some minor ui tweaks

#rnx

diff --git a/Engine/Plugins/Developer/CustomLaunchUI/Source/CustomLaunchUI.Build.cs b/Engine/Plugins/Developer/CustomLaunchUI/Source/CustomLaunchUI.Build.cs
--- a/Engine/Plugins/Developer/CustomLaunchUI/Source/CustomLaunchUI.Build.cs
+++ b/Engine/Plugins/Developer/CustomLaunchUI/Source/CustomLaunchUI.Build.cs
@@ -22,6 +22,7 @@ public class CustomLaunchUI : ModuleRules
 			"TargetDeviceServices",
 			"Projects",
 			"DeveloperToolSettings",
+			"Zen",
 		]);
 	}
 }
diff --git a/Engine/Plugins/Developer/CustomLaunchUI/Source/Private/Marshalling/BasicProfileMarshaller.cpp b/Engine/Plugins/Developer/CustomLaunchUI/Source/Private/Marshalling/BasicProfileMarshaller.cpp
--- a/Engine/Plugins/Developer/CustomLaunchUI/Source/Private/Marshalling/BasicProfileMarshaller.cpp
+++ b/Engine/Plugins/Developer/CustomLaunchUI/Source/Private/Marshalling/BasicProfileMarshaller.cpp
@@ -44,7 +44,7 @@ protected:
 	EBuildConfiguration GetBuildConfiguration() const;
 
 	void SetContentScheme(CustomLaunchUI::EContentScheme ContentScheme);
-	bool IsContentSchemeAvailable(CustomLaunchUI::EContentScheme) const;
+	bool IsContentSchemeAvailable(CustomLaunchUI::EContentScheme, FText& OutReason ) const;
 
 	FText GetCommandLine() const;
 	void SetCommandLine( const FText& NewText, ETextCommit::Type InTextCommit );
@@ -205,10 +205,32 @@ void FBasicProfileMarshaller::SetContentScheme(CustomLaunchUI::EContentScheme In
 	ContentScheme = Model->DetermineProfileContentScheme(Profile);
 }
 
-bool FBasicProfileMarshaller::IsContentSchemeAvailable(CustomLaunchUI::EContentScheme InContentScheme) const
+bool FBasicProfileMarshaller::IsContentSchemeAvailable(CustomLaunchUI::EContentScheme InContentScheme, FText& OutReason) const
 {
-	// @fixme: check the project supports the content scheme etc
-	return InContentScheme != CustomLaunchUI::EContentScheme::ZenPakStreaming;
+	CustomLaunchUI::FProjectSettings ProjectSettings = Model->GetProjectSettings(Profile);
+
+	// loose files can't be selected if the project is using zen store because there's no way to opt-out of Zen Store from the UAT command line.
+	if (InContentScheme == CustomLaunchUI::EContentScheme::LooseFiles && ProjectSettings.bUseZenStore)
+	{
+		OutReason = LOCTEXT("NoLooseFilesReason", "Loose Files cannot be used when Zen Store is enabled in Project Settings");
+		return false;
+	}
+
+	// cannot launch via Zen if we're targeting a remote device and Zen isn't accepting external connections
+	if (InContentScheme == CustomLaunchUI::EContentScheme::ZenStreaming && !ProjectSettings.bAllowRemoteNetworkService && !Model->IsHostPlatform(Profile))
+	{
+		OutReason = LOCTEXT("NoZenReason", "Zen Streaming to a remote device requires AllowRemoteNetworkService");
+		return false;
+	}
+
+	// basic launch is aimed at launching the current content, not an external build
+	if (InContentScheme == CustomLaunchUI::EContentScheme::ZenPakStreaming)
+	{
+		// don't set a reason - just hide the item for Basic Launch
+		return false;
+	}
+
+	return true;
 }
 
 FText FBasicProfileMarshaller::GetCommandLine() const
diff --git a/Engine/Plugins/Developer/CustomLaunchUI/Source/Private/Marshalling/CustomProfileMarshaller.cpp b/Engine/Plugins/Developer/CustomLaunchUI/Source/Private/Marshalling/CustomProfileMarshaller.cpp
--- a/Engine/Plugins/Developer/CustomLaunchUI/Source/Private/Marshalling/CustomProfileMarshaller.cpp
+++ b/Engine/Plugins/Developer/CustomLaunchUI/Source/Private/Marshalling/CustomProfileMarshaller.cpp
@@ -72,7 +72,7 @@ protected:
 	EBuildConfiguration GetBuildConfiguration() const;
 
 	void SetContentScheme(CustomLaunchUI::EContentScheme ContentScheme);
-	bool IsContentSchemeAvailable(CustomLaunchUI::EContentScheme) const;
+	bool IsContentSchemeAvailable(CustomLaunchUI::EContentScheme, FText& OutReason) const;
 
 	FString GetCommandLine() const;
 	void SetCommandLine( const FString& NewCommandLine );
@@ -298,14 +298,15 @@ void FCustomProfileMarshaller::Construct()
 
 	TreeData->AddHeading( LOCTEXT("DirectoryHeading", "Directory") )
 
-		.AddDirectoryString( LOCTEXT("CustomStagingPathLabel", "Custom Stage Directory"),
-			{
-				.GetValue = [this]()				{ return Profile->GetPackageDirectory(); },
-				.SetValue = [this](FString Value)	{ Profile->SetPackageDirectory(Value); },
-				.GetDefaultValue = EmptyString,
-				.IsEnabled = ForCooked,
-			}
-		)
+		// unlikely to want to change this in most use cases
+		//.AddDirectoryString( LOCTEXT("CustomStagingPathLabel", "Custom Stage Directory"),
+		//	{
+		//		.GetValue = [this]()				{ return Profile->GetPackageDirectory(); },
+		//		.SetValue = [this](FString Value)	{ Profile->SetPackageDirectory(Value); },
+		//		.GetDefaultValue = EmptyString,
+		//		.IsEnabled = ForCooked,
+		//	}
+		//)
 
 		.AddBoolean( LOCTEXT("ArchiveBuildLabel", "Archive Build"),
 			{
@@ -367,14 +368,15 @@ void FCustomProfileMarshaller::Construct()
 			}
 		)
 
-		.AddBoolean( LOCTEXT("BuildUATLabel", "Build UAT"),
-			{
-				.GetValue = [this]()			{ return Profile->IsBuildingUAT(); },
-				.SetValue = [this](bool bValue)	{ Profile->SetBuildUAT(bValue); },
-				.GetDefaultValue = [this]()		{ return DefaultProfile->IsBuildingUAT(); },
-				.IsEnabled = ForCode,
-			}
-		)
+		// hiding this option for now - no good reason to buid UAT it for the majority of normal use cases
+		//.AddBoolean( LOCTEXT("BuildUATLabel", "Build UAT"),
+		//	{
+		//		.GetValue = [this]()			{ return Profile->IsBuildingUAT(); },
+		//		.SetValue = [this](bool bValue)	{ Profile->SetBuildUAT(bValue); },
+		//		.GetDefaultValue = [this]()		{ return DefaultProfile->IsBuildingUAT(); },
+		//		.IsEnabled = ForCode,
+		//	}
+		//)
 
 
 	;
@@ -392,14 +394,15 @@ void FCustomProfileMarshaller::Construct()
 			}
 		)
 
-		.AddBoolean( LOCTEXT("DeployModifiedLabel", "Only Deploy Modified Content"),
-			{
-				.GetValue = [this]()			{ return Profile->IsDeployingIncrementally(); },
-				.SetValue = [this](bool bValue)	{ Profile->SetIncrementalDeploying(bValue); },
-				.GetDefaultValue = [this]()		{ return DefaultProfile->IsDeployingIncrementally(); },
-				.IsEnabled = ForDeployment,
-			}
-		)
+		// not relevant to most scenarios now
+		//.AddBoolean( LOCTEXT("DeployModifiedLabel", "Only Deploy Modified Content"),
+		//	{
+		//		.GetValue = [this]()			{ return Profile->IsDeployingIncrementally(); },
+		//		.SetValue = [this](bool bValue)	{ Profile->SetIncrementalDeploying(bValue); },
+		//		.GetDefaultValue = [this]()		{ return DefaultProfile->IsDeployingIncrementally(); },
+		//		.IsEnabled = ForDeployment,
+		//	}
+		//)
 
 		.AddWidget( LOCTEXT("TargetDeviceLabel", "Target Device"), 
 			{
@@ -545,10 +548,6 @@ void FCustomProfileMarshaller::SetProjectName(FString ProjectPath)
 	Profile->SetProjectSpecified(!ProjectPath.IsEmpty());
 	Profile->SetProjectPath(ProjectPath);
 
-	// updating the project can reset the Taret
-	//Target = Profile->HasBuildTargetSpecified() && !Profile->GetBuildTarget().IsEmpty() ? Profile->GetBuildTarget() : FString();
-	//bHasTargetSpecified = Profile->HasBuildTargetSpecified();
-
 	if (CustomLaunchUI::bUseFriendlyBuildTargetSelection)
 	{
 		Model->UpdateCookedPlatformsFromBuildTarget(Profile);
@@ -639,9 +638,33 @@ void FCustomProfileMarshaller::SetContentScheme(CustomLaunchUI::EContentScheme I
 	ContentScheme = Model->DetermineProfileContentScheme(Profile);
 }
 
-bool FCustomProfileMarshaller::IsContentSchemeAvailable(CustomLaunchUI::EContentScheme InContentScheme) const
+bool FCustomProfileMarshaller::IsContentSchemeAvailable(CustomLaunchUI::EContentScheme InContentScheme, FText& OutReason) const
 {
-	// @fixme: check the project supports the content scheme etc
+	CustomLaunchUI::FProjectSettings ProjectSettings = Model->GetProjectSettings(Profile);
+
+	// loose files can't be selected if the project is using zen store because there's no way to opt-out of Zen Store from the UAT command line.
+	if (InContentScheme == CustomLaunchUI::EContentScheme::LooseFiles && ProjectSettings.bUseZenStore)
+	{
+		OutReason = LOCTEXT("NoLooseFilesReason", "Loose Files cannot be used when Zen Store is enabled in Project Settings");
+		return false;
+	}
+
+	// don't show zen pak streaming option if it isn't going to be set up automatically by UAT
+	// @todo: could potentially look at shelling out to Zen & querying if we have a dynamic workspace? no support for async config queries in this tool yet though
+	if (InContentScheme == CustomLaunchUI::EContentScheme::ZenPakStreaming && !ProjectSettings.bHasAutomaticZenPakStreamingWorkspaceCreation)
+	{
+		OutReason = LOCTEXT("NoZenPakReason", "Automatic Zen Pak streaming workspace creation has not been enabled in Project Settings");
+		return false;
+	}
+	
+	// cannot launch via Zen if we're targeting a remote device and Zen isn't accepting external connections
+	if ((InContentScheme == CustomLaunchUI::EContentScheme::ZenStreaming || InContentScheme == CustomLaunchUI::EContentScheme::ZenPakStreaming) 
+		&& !ProjectSettings.bAllowRemoteNetworkService && !Model->IsHostPlatform(Profile))
+	{
+		OutReason = LOCTEXT("NoZenReason", "Zen Streaming to a remote device requires AllowRemoteNetworkService");
+		return false;
+	}
+
 	return true;
 }
 
diff --git a/Engine/Plugins/Developer/CustomLaunchUI/Source/Private/Model/CustomLaunchUIModel.cpp b/Engine/Plugins/Developer/CustomLaunchUI/Source/Private/Model/CustomLaunchUIModel.cpp
--- a/Engine/Plugins/Developer/CustomLaunchUI/Source/Private/Model/CustomLaunchUIModel.cpp
+++ b/Engine/Plugins/Developer/CustomLaunchUI/Source/Private/Model/CustomLaunchUIModel.cpp
@@ -3,6 +3,7 @@
 #include "Model/CustomLaunchUIModel.h"
 #include "Misc/Paths.h"
 #include "Misc/ConfigCacheIni.h"
+#include "Misc/ConfigContext.h"
 #include "Misc/App.h"
 #include "PlatformInfo.h"
 #include "GameProjectHelper.h"
@@ -12,7 +13,8 @@
 #include "Serialization/BufferArchive.h"
 #include "Serialization/MemoryReader.h"
 #include "Modules/ModuleManager.h"
-
+#include "Settings/ProjectPackagingSettings.h"
+#include "Experimental/ZenServerInterface.h"
 
 #define LOCTEXT_NAMESPACE "SCustomLaunchProfileSelector"
 
@@ -175,6 +177,31 @@ namespace CustomLaunchUI
 		return nullptr;
 	}
 
+	bool FModel::IsHostPlatform(const ILauncherProfilePtr& Profile)
+	{
+		if (Profile.IsValid() && Profile->GetCookedPlatforms().Num() > 0)
+		{
+			FString SelectedPlatform = Profile->GetCookedPlatforms()[0];
+			return IsHostPlatform(FName(SelectedPlatform));
+		}
+
+		return false;
+	}
+
+	bool FModel::IsHostPlatform(FName PlatformName)
+	{
+		const PlatformInfo::FTargetPlatformInfo* PlatformInfo = PlatformInfo::FindPlatformInfo(PlatformName);
+
+		if (PlatformInfo != nullptr && PlatformInfo->IniPlatformName == FPlatformProperties::IniPlatformName())
+		{
+			return true;
+		}
+
+		return false;
+	}
+
+
+
 	FTargetInfo FModel::GetBuildTargetInfo( const FString& BuildTargetName, const FString& ProjectPath )
 	{
 		FTargetInfo Result;
@@ -470,9 +497,11 @@ namespace CustomLaunchUI
 		Profile->SetDeployWithUnrealPak(bPakFiles);
 		Profile->SetUseZenStore(bUseZen);
 
+		ELauncherProfileDeploymentModes::Type DeploymentMode = DefaultDeploymentMode;
 		if (bCOTF)
 		{
 			Profile->SetCookMode(ELauncherProfileCookModes::OnTheFly);
+			DeploymentMode = ELauncherProfileDeploymentModes::FileServer;
 		}
 		else if (bZenPakStreaming)
 		{
@@ -490,12 +519,84 @@ namespace CustomLaunchUI
 		}
 		else
 		{
-			Profile->SetDeploymentMode(bCOTF ? ELauncherProfileDeploymentModes::FileServer : DefaultDeploymentMode);
+			Profile->SetDeploymentMode(DeploymentMode);
 			Profile->SetPackagingMode(ELauncherProfilePackagingModes::DoNotPackage);
 		}
 	}
 
 
+	void FModel::ReadProjectSettingsFromConfig( FConfigCacheIni& InConfig, FProjectSettings& OutResult )
+	{
+		// read project packaging settings
+		const TCHAR* ProjectPackagingConfigSection = TEXT("/Script/UnrealEd.ProjectPackagingSettings");
+		InConfig.GetBool(ProjectPackagingConfigSection, TEXT("bUseZenStore"), OutResult.bUseZenStore, GGameIni);
+		InConfig.GetBool(ProjectPackagingConfigSection, TEXT("bEnablePakStreaming"), OutResult.bHasAutomaticZenPakStreamingWorkspaceCreation, GGameIni);
+		
+		// read zen settings
+#if UE_WITH_ZEN
+		using namespace UE::Zen;
+		const FServiceSettings& ZenServiceSettings = GetDefaultServiceInstance().GetServiceSettings();
+		OutResult.bAllowRemoteNetworkService = ZenServiceSettings.IsAutoLaunch() && ZenServiceSettings.SettingsVariant.Get<FServiceAutoLaunchSettings>().bAllowRemoteNetworkService;
+#else
+		const TCHAR* ZenAutoLaunchSettings = TEXT("Zen.AutoLaunch");
+		InConfig.GetBool(ZenAutoLaunchSettings, TEXT("AllowRemoteNetworkService"), OutResult.bAllowRemoteNetworkService, GEngineIni);
+#endif
+	}
+
+	const FProjectSettings FModel::GetProjectSettings( const FString& InProjectPath )
+	{
+		FProjectSettings Result;
+
+		FString ProjectPath = FPaths::ConvertRelativePathToFull(InProjectPath);
+		FString ProjectName = FPaths::GetBaseFilename(ProjectPath);
+
+		FString GlobalProjectPath;
+		if (FPaths::IsProjectFilePathSet())
+		{
+			GlobalProjectPath = FPaths::ConvertRelativePathToFull(FPaths::GetProjectFilePath());
+		}
+
+		if (!GlobalProjectPath.IsEmpty() && GlobalProjectPath == ProjectPath )
+		{
+			// read the current project's ini settings
+			Result.bIsCurrentEditorProject = true;
+			ReadProjectSettingsFromConfig(*GConfig, Result);
+
+			// apply any in-memory properties @fixme: is this necessary?
+			const UProjectPackagingSettings* ProjectPackagingSettings = UProjectPackagingSettings::StaticClass()->GetDefaultObject<UProjectPackagingSettings>();
+			Result.bUseZenStore = ProjectPackagingSettings->bUseZenStore;
+
+			// note this is not cached because the user could edit the properties from within the editor
+		}
+		else if (CachedProjectSettings.Contains(ProjectName))
+		{
+			Result = CachedProjectSettings[ProjectName];
+		}
+		else if (!ProjectPath.IsEmpty())
+		{
+			// load the other project's ini files into a temporary config cache
+			FConfigCacheIni Config(EConfigCacheType::Temporary);
+			FConfigContext Context = FConfigContext::ReadIntoConfigSystem(&Config, FString());
+			Context.ProjectConfigDir = FPaths::Combine(FPaths::GetPath(ProjectPath), TEXT("Config/"));
+			Config.InitializeKnownConfigFiles(Context);
+
+			// read the ini settings
+			Result.bIsCurrentEditorProject = false;
+			ReadProjectSettingsFromConfig(Config, Result);
+
+			// cache the result
+			CachedProjectSettings.Add( ProjectName, Result);
+		}
+
+		return MoveTemp(Result);
+	}
+
+	const FProjectSettings FModel::GetProjectSettings( const ILauncherProfileRef& Profile )
+	{
+		return GetProjectSettings(*Profile->GetProjectPath());
+	}
+
+
 
 	bool FModel::IsBasicLaunchProfile( const ILauncherProfilePtr& Profile ) const
 	{
diff --git a/Engine/Plugins/Developer/CustomLaunchUI/Source/Private/Model/CustomLaunchUIModel.h b/Engine/Plugins/Developer/CustomLaunchUI/Source/Private/Model/CustomLaunchUIModel.h
--- a/Engine/Plugins/Developer/CustomLaunchUI/Source/Private/Model/CustomLaunchUIModel.h
+++ b/Engine/Plugins/Developer/CustomLaunchUI/Source/Private/Model/CustomLaunchUIModel.h
@@ -51,6 +51,20 @@ namespace CustomLaunchUI
 		{}
 	};
 
+	struct FProjectSettings
+	{
+		// whether this project is the one that is currently open in the editor
+		bool bIsCurrentEditorProject = false;
+
+		// whether Zen Store is enabled via the project settings - in this case, we can't use Loose Files
+		bool bUseZenStore = false;
+
+		// whether a workspace will be created automatically when using Zen Pak streaming
+		bool bHasAutomaticZenPakStreamingWorkspaceCreation = false;
+
+		// whether Zen Server is allowed to accept connections from remote machines (i.e. console devkits, phones etc)
+		bool bAllowRemoteNetworkService = false;
+	};
 
 
 	DECLARE_DELEGATE_OneParam(FOnProfileClicked, const ILauncherProfilePtr&)
@@ -130,6 +144,9 @@ namespace CustomLaunchUI
 		static const PlatformInfo::FTargetPlatformInfo* GetPlatformInfo(const ILauncherProfilePtr& Profile);
 		static const PlatformInfo::FTargetPlatformInfo* GetPlatformInfo(FName PlatformName, const FTargetInfo& BuildTargetInfo);
 
+		static bool IsHostPlatform(const ILauncherProfilePtr& Profile);
+		static bool IsHostPlatform(FName PlatformName);
+
 		static FTargetInfo GetBuildTargetInfo( const FString& BuildTargetName, const FString& ProjectPath );
 		static FTargetInfo GetBuildTargetInfo( const ILauncherProfileRef& Profile );
 
@@ -140,6 +157,9 @@ namespace CustomLaunchUI
 		static FString GetVanillaPlatformName( const FString& PlatformName );
 		static FString GetBuildTargetPlatformName( const FString& PlatformName, const FTargetInfo& BuildTargetInfo );
 
+		const FProjectSettings GetProjectSettings( const FString& ProjectPath );
+		const FProjectSettings GetProjectSettings( const ILauncherProfileRef& Profile );
+
 	private:
 		void LoadConfig();
 		void SaveConfig();
@@ -149,6 +169,8 @@ namespace CustomLaunchUI
 		void HandleDeviceProxyAdded(const TSharedRef<ITargetDeviceProxy>& DeviceProxy );
 		void HandleDeviceProxyRemoved(const TSharedRef<ITargetDeviceProxy>& DeviceProxy );
 
+		void ReadProjectSettingsFromConfig( FConfigCacheIni& InConfig, FProjectSettings& OutResult );
+
 		TSharedRef<ITargetDeviceProxyManager> DeviceProxyManager;
 		TSharedRef<ILauncher> Launcher;
 		TSharedRef<ILauncherProfileManager> ProfileManager;
@@ -161,6 +183,8 @@ namespace CustomLaunchUI
 		friend class FLaunchLogTextLayoutMarshaller;
 		TArray<TSharedPtr<FLaunchLogMessage>> LaunchLogMessages;
 
+		TMap<FString,FProjectSettings> CachedProjectSettings;
+
 		FOnSelectedProfileChanged ProfileSelectedDelegate;
 		FOnSelectedPlatformChanged PlatformChangedDelegate;
 
diff --git a/Engine/Plugins/Developer/CustomLaunchUI/Source/Private/Widgets/Shared/SCustomLaunchContentSchemeCombo.cpp b/Engine/Plugins/Developer/CustomLaunchUI/Source/Private/Widgets/Shared/SCustomLaunchContentSchemeCombo.cpp
--- a/Engine/Plugins/Developer/CustomLaunchUI/Source/Private/Widgets/Shared/SCustomLaunchContentSchemeCombo.cpp
+++ b/Engine/Plugins/Developer/CustomLaunchUI/Source/Private/Widgets/Shared/SCustomLaunchContentSchemeCombo.cpp
@@ -42,22 +42,42 @@ TSharedRef<SWidget> SCustomLaunchContentSchemeCombo::MakeContentSchemeSelectionW
 	{
 		for (CustomLaunchUI::EContentScheme ContentScheme : CustomLaunchUI::GetAllContentSchemes())
 		{
-			if (IsContentSchemeAvailable.IsBound() && !IsContentSchemeAvailable.Execute(ContentScheme))
+			FText Reason = FText::GetEmpty();
+			bool bIsAvailable = !IsContentSchemeAvailable.IsBound() || IsContentSchemeAvailable.Execute(ContentScheme, Reason);
+
+			if (bIsAvailable)
+			{
+				MenuBuilder.AddMenuEntry(
+					CustomLaunchUI::GetContentSchemeDisplayName(ContentScheme),
+					CustomLaunchUI::GetContentSchemeToolTip(ContentScheme),
+					FSlateIcon(), 
+					FUIAction(
+						FExecuteAction::CreateSP(this, &SCustomLaunchContentSchemeCombo::SetContentScheme, ContentScheme),
+						FCanExecuteAction(),
+						FGetActionCheckState::CreateLambda( [this, ContentScheme]() { return (SelectedContentScheme.Get() == ContentScheme) ? ECheckBoxState::Checked : ECheckBoxState::Unchecked; })
+					),
+					NAME_None,
+					EUserInterfaceActionType::Check);
+			}
+			else if (Reason.IsEmpty())
 			{
+				// no reason given - hide the item completely
 				continue;
 			}
-
-			MenuBuilder.AddMenuEntry(
-				CustomLaunchUI::GetContentSchemeDisplayName(ContentScheme),
-				CustomLaunchUI::GetContentSchemeToolTip(ContentScheme),
-				FSlateIcon(), 
-				FUIAction(
-					FExecuteAction::CreateSP(this, &SCustomLaunchContentSchemeCombo::SetContentScheme, ContentScheme),
-					FCanExecuteAction(),
-					FGetActionCheckState::CreateLambda( [this, ContentScheme]() { return (SelectedContentScheme.Get() == ContentScheme) ? ECheckBoxState::Checked : ECheckBoxState::Unchecked; })
-				),
-				NAME_None,
-				EUserInterfaceActionType::Check);
+			else
+			{
+				MenuBuilder.AddMenuEntry(
+					CustomLaunchUI::GetContentSchemeDisplayName(ContentScheme),
+					Reason,
+					FSlateIcon(), 
+					FUIAction(
+						FExecuteAction(),
+						FCanExecuteAction::CreateLambda( []() { return false; } ),
+						FGetActionCheckState::CreateLambda( [this, ContentScheme]() { return (SelectedContentScheme.Get() == ContentScheme) ? ECheckBoxState::Checked : ECheckBoxState::Unchecked; })
+					),
+					NAME_None,
+					EUserInterfaceActionType::Check);
+			}
 		}
 	}
 
diff --git a/Engine/Plugins/Developer/CustomLaunchUI/Source/Private/Widgets/Shared/SCustomLaunchContentSchemeCombo.h b/Engine/Plugins/Developer/CustomLaunchUI/Source/Private/Widgets/Shared/SCustomLaunchContentSchemeCombo.h
--- a/Engine/Plugins/Developer/CustomLaunchUI/Source/Private/Widgets/Shared/SCustomLaunchContentSchemeCombo.h
+++ b/Engine/Plugins/Developer/CustomLaunchUI/Source/Private/Widgets/Shared/SCustomLaunchContentSchemeCombo.h
@@ -15,7 +15,7 @@ class SCustomLaunchContentSchemeCombo
 {
 public:
 	DECLARE_DELEGATE_OneParam(FOnSelectionChanged, CustomLaunchUI::EContentScheme );
-	DECLARE_DELEGATE_RetVal_OneParam(bool, FGetContentSchemeAvailability, CustomLaunchUI::EContentScheme )
+	DECLARE_DELEGATE_RetVal_TwoParams(bool, FGetContentSchemeAvailability, CustomLaunchUI::EContentScheme, FText& )
 
 	SLATE_BEGIN_ARGS(SCustomLaunchContentSchemeCombo) 
 		: _TextStyle( &FCoreStyle::Get().GetWidgetStyle<FTextBlockStyle>( "NormalText" ) )
diff --git a/Engine/Plugins/Developer/CustomLaunchUI/Source/Private/Widgets/Shared/SCustomLaunchProjectCombo.cpp b/Engine/Plugins/Developer/CustomLaunchUI/Source/Private/Widgets/Shared/SCustomLaunchProjectCombo.cpp
--- a/Engine/Plugins/Developer/CustomLaunchUI/Source/Private/Widgets/Shared/SCustomLaunchProjectCombo.cpp
+++ b/Engine/Plugins/Developer/CustomLaunchUI/Source/Private/Widgets/Shared/SCustomLaunchProjectCombo.cpp
@@ -84,7 +84,7 @@ TSharedRef<SWidget> SCustomLaunchProjectCombo::MakeProjectSelectionWidget()
 			else if (!ProjectName.IsEmpty())
 			{
 				FUIAction CurrentProjectAction(FExecuteAction::CreateSP(this, &SCustomLaunchProjectCombo::SetProjectPath, ProjectPath));
-				MenuBuilder.AddMenuEntry(FText::FromString(*ProjectName), FText::FromString(ProjectPath), FSlateIcon(), CurrentProjectAction);
+				MenuBuilder.AddMenuEntry(FText::Format(LOCTEXT("CurrentProjectActionFmt", "Current Project ({0})"), FText::FromString(ProjectName)), LOCTEXT("CurrentProjectHint", "Use the current project"), FSlateIcon(), CurrentProjectAction);
 			}
 		}
 
diff --git a/Engine/Source/Developer/LauncherServices/Private/Launcher/LauncherWorker.cpp b/Engine/Source/Developer/LauncherServices/Private/Launcher/LauncherWorker.cpp
--- a/Engine/Source/Developer/LauncherServices/Private/Launcher/LauncherWorker.cpp
+++ b/Engine/Source/Developer/LauncherServices/Private/Launcher/LauncherWorker.cpp
@@ -594,19 +594,11 @@ FString FLauncherWorker::CreateUATCommand( const ILauncherProfileRef& InProfile,
 		*SessionName);
 	}
 
-	// zen pak streaming commandline
-	FString ZenPakStreamingCommands;
-	if (InProfile->IsUsingZenPakStreaming())
-	{
-		ZenPakStreamingCommands = FString::Printf(TEXT("-ZenPak -ZenWorkspaceSharePath=\\\"%s\\\" "), *InProfile->GetZenPakStreamingPath() );
-	}
-
 	FString AdditionalCommandLine;
-	FString AddCmdLine = FString::Printf(TEXT("%s%s %s %s %s"),
+	FString AddCmdLine = FString::Printf(TEXT("%s%s %s %s"),
 		*SessionCommands,
 		*RoleCommands,
 		*LocalizationCommands,
-		*ZenPakStreamingCommands,
 		*InProfile->GetAdditionalCommandLineParameters());
 	if (!AddCmdLine.TrimStartAndEnd().IsEmpty()) // don't bother with -addcmdline if it's empty
 	{
@@ -896,9 +888,7 @@ FString FLauncherWorker::CreateUATCommand( const ILauncherProfileRef& InProfile,
 		UATCommand += TEXT(" -skipcook");
 		if (InProfile->IsUsingZenPakStreaming())
 		{
-			// @fixme: should auto-configure the zen dynamic workspace in UAT?
-			//UATCommand += TEXT(" -createdynamiczenworkspace=");
-			//UATCommand += InProfile->GetZenPakStreamingPath();
+			UATCommand += FString::Printf(TEXT(" -skippak -ZenWorkspaceSharePath=\\\"%s\\\" "), *InProfile->GetZenPakStreamingPath() );
 		}
 		break;
 	}
