Streamline some custom profile fields:
- Add a 'cook content' tickbox
- disallow incremental cooking when importing a snapshot because the snapshot import currently also does a different incremental cook & the two features don't play nicely together
- property editor will hide a section header if all of its children are hidden
- hide all cooking & archiving properties if they're not needed
- move archive section lower in the list
- use 'build' rather than 'code'
- allow configuration selection even if not building.

#rnx

diff --git a/Engine/Plugins/Developer/CustomLaunchUI/Source/Private/Marshalling/CustomProfileMarshaller.cpp b/Engine/Plugins/Developer/CustomLaunchUI/Source/Private/Marshalling/CustomProfileMarshaller.cpp
--- a/Engine/Plugins/Developer/CustomLaunchUI/Source/Private/Marshalling/CustomProfileMarshaller.cpp
+++ b/Engine/Plugins/Developer/CustomLaunchUI/Source/Private/Marshalling/CustomProfileMarshaller.cpp
@@ -80,6 +80,9 @@ protected:
 	void SetSelectedPlatforms( TArray<FString> SelectedPlatforms );
 	TArray<FString> GetSelectedPlatforms() const;
 
+	void SetCook( bool bCook );
+	bool GetCook( ILauncherProfilePtr InProfile = nullptr ) const;
+
 	void SetMapsToCook(TArray<FString> MapsToCook);
 	TArray<FString> GetMapsToCook() const;
 	EMapOption GetMapOption() const;
@@ -100,11 +103,11 @@ protected:
 	TSharedRef<SWidget> CreateDeployDeviceWidget();
 	float DeployDeviceListHeight = 150.0f;
 
-	void SetBuildCode( bool bBuildCode );
-	bool GetBuildCode( ILauncherProfilePtr InProfile = nullptr ) const;
+	void SetBuild( bool bBuild );
+	bool GetBuild( ILauncherProfilePtr InProfile = nullptr ) const;
 
-	void SetBuildAlways( bool bBuildAlways );
-	bool GetBuildAlways( ILauncherProfilePtr InProfile = nullptr ) const;
+	void SetForceBuild( bool bForceBuild );
+	bool GetForceBuild( ILauncherProfilePtr InProfile = nullptr ) const;
 
 	void SetDeployToDevice( bool bDeployToDevice );
 	bool GetDeployToDevice( ILauncherProfilePtr InProfile = nullptr ) const;
@@ -119,6 +122,7 @@ protected:
 
 protected:
 	virtual void OnPropertyChanged() override;
+	void RefreshContentScheme();
 	void CacheStartupMapList();
 	TSharedRef<SWidget> OnGenerateComboWidget( TSharedPtr<FString> InComboString );
 
@@ -129,6 +133,7 @@ protected:
 	const TSharedRef<CustomLaunchUI::FModel> Model;
 
 	CustomLaunchUI::EContentScheme ContentScheme;
+	bool bShouldCook;
 	TArray<TSharedPtr<FString>> CachedStartupMaps;
 	TSharedPtr<SSearchableComboBox> InitalMapCombo;
 	TSharedPtr<SCustomLaunchMapListView> MapListView;
@@ -162,9 +167,14 @@ void FCustomProfileMarshaller::Construct()
 		return (ContentScheme != CustomLaunchUI::EContentScheme::ZenPakStreaming && ContentScheme != CustomLaunchUI::EContentScheme::CookOnTheFly);
 	};
 
+	auto ForEnabledCooked = [this]()
+	{
+		return (bShouldCook && ContentScheme != CustomLaunchUI::EContentScheme::ZenPakStreaming && ContentScheme != CustomLaunchUI::EContentScheme::CookOnTheFly);
+	};
+
 	auto ForCode = [this]()
 	{
-		return GetBuildCode();
+		return GetBuild();
 	};
 
 	auto ForDeployment = [this]()
@@ -219,9 +229,6 @@ void FCustomProfileMarshaller::Construct()
 		)
 
 		.AddWidget( LOCTEXT("ConfigurationLabel", "Configuration"), 
-			{
-				.IsEnabled = ForCode,
-			},
 			SNew(SCustomLaunchBuildConfigurationCombo)
 			.OnSelectionChanged(this, &FCustomProfileMarshaller::SetBuildConfiguration)
 			.SelectedBuildConfiguration(this, &FCustomProfileMarshaller::GetBuildConfiguration)
@@ -258,7 +265,7 @@ void FCustomProfileMarshaller::Construct()
 			}
 		)
 
-		.AddBoolean( LOCTEXT("GenerateChunksLabel", "GenerateChunks"),
+		.AddBoolean( LOCTEXT("GenerateChunksLabel", "Generate Chunks"),
 			{
 				.GetValue = [this]()			{ return Profile->IsGeneratingChunks(); },
 				.SetValue = [this](bool bValue)	{ Profile->SetGenerateChunks(bValue); },
@@ -287,53 +294,30 @@ void FCustomProfileMarshaller::Construct()
 
 		.AddBoolean( LOCTEXT("IncrementalCookLabel", "Incremental Cook"),
 			{
-				.GetValue = [this]()			{ return Profile->IsCookingIncrementally(); },
+				.GetValue = [this]()			{ return Profile->IsCookingIncrementally() || Profile->IsImportingZenSnapshot(); }, // show it as ticked when snapshot import is enabled for now to indicate that it will be doing something incremental. eventually this is likely to become a separate field
 				.SetValue = [this](bool bValue)	{ Profile->SetIncrementalCooking(bValue); },
 				.GetDefaultValue = [this]()		{ return DefaultProfile->IsCookingIncrementally(); },
 				.IsVisible = ForCooked,
+				.IsEnabled = [this]()			{ return !Profile->IsImportingZenSnapshot(); }, // zen snapshot import & incremental cooking can't be used together. '-iterativecook' is the 'old' system, and zen import will enable the new '-cookincremental' & the two flags don't play nicely together
 			}
 		)
 	;
 
 
-	TreeData->AddHeading( LOCTEXT("DirectoryHeading", "Directory") )
-
-		// unlikely to want to change this in most use cases
-		//.AddDirectoryString( LOCTEXT("CustomStagingPathLabel", "Custom Stage Directory"),
-		//	{
-		//		.GetValue = [this]()				{ return Profile->GetPackageDirectory(); },
-		//		.SetValue = [this](FString Value)	{ Profile->SetPackageDirectory(Value); },
-		//		.GetDefaultValue = EmptyString,
-		//		.IsEnabled = ForCooked,
-		//	}
-		//)
-
-		.AddBoolean( LOCTEXT("ArchiveBuildLabel", "Archive Build"),
-			{
-				.GetValue = [this]()			{ return Profile->IsArchiving(); },
-				.SetValue = [this](bool bValue)	{ Profile->SetArchive(bValue); },
-				.GetDefaultValue = [this]()		{ return DefaultProfile->IsArchiving(); },
-				.IsEnabled = ForCooked,
-			}
-		)
+	TreeData->AddHeading( LOCTEXT("CookingHeading", "Maps And Cooking") )
 
-		.AddDirectoryString( LOCTEXT("ArchivePathLabel", "Archive Directory"),
+		.AddBoolean( LOCTEXT("CookLabel", "Cook Content"),
 			{
-				.GetValue = [this]()				{ return Profile->GetArchiveDirectory(); },
-				.SetValue = [this](FString Value)	{ Profile->SetArchiveDirectory(Value); },
-				.GetDefaultValue = EmptyString,
-				.IsEnabled = [this,ForCooked]()		{ return Profile->IsArchiving() && ForCooked(); }
+				.GetValue = [this]()			{ return GetCook(); },
+				.SetValue = [this](bool bValue)	{ SetCook(bValue); },
+				.GetDefaultValue = [this]()		{ return GetCook(DefaultProfile); },
+				.IsVisible = ForCooked
 			}
 		)
-	;
-
-
-
-	TreeData->AddHeading( LOCTEXT("CookingHeading", "Maps And Cooking") )
 
 		.AddWidget( LOCTEXT("MapsToCookLabel", "Maps To Cook"), 
 			{
-				.IsEnabled = ForCooked,
+				.IsVisible = ForEnabledCooked,
 			},
 			CreateMapListWidget()
 		)
@@ -343,27 +327,27 @@ void FCustomProfileMarshaller::Construct()
 				.GetValue = [this]()				{ return Profile->GetCookOptions(); },
 				.SetValue = [this](FString Value)	{ Profile->SetCookOptions( Value ); },
 				.GetDefaultValue = EmptyString,
-				.IsEnabled = ForCooked,
+				.IsVisible = ForEnabledCooked,
 			}
 		)
 	;
 
 
-	TreeData->AddHeading( LOCTEXT("CodeHeading", "Code") )
+	TreeData->AddHeading( LOCTEXT("BuildHeading", "Build") )
 
-		.AddBoolean( LOCTEXT("BuildCodeLabel", "Build Code"),
+		.AddBoolean( LOCTEXT("BuildLabel", "Build the game"),
 			{
-				.GetValue = [this]()			{ return GetBuildCode(); },
-				.SetValue = [this](bool bValue)	{ SetBuildCode(bValue); },
-				.GetDefaultValue = [this]()		{ return GetBuildCode(DefaultProfile); },
+				.GetValue = [this]()			{ return GetBuild(); },
+				.SetValue = [this](bool bValue)	{ SetBuild(bValue); },
+				.GetDefaultValue = [this]()		{ return GetBuild(DefaultProfile); },
 			}
 		)
 
-		.AddBoolean( LOCTEXT("ForceBuildCodeLabel", "Build even if pre-build target exists"),
+		.AddBoolean( LOCTEXT("ForceBuildLabel", "Build even if a pre-built target exists"),
 			{
-				.GetValue = [this]()			{ return GetBuildAlways(); },
-				.SetValue = [this](bool bValue)	{ SetBuildAlways(bValue); },
-				.GetDefaultValue = [this]()		{ return GetBuildAlways(DefaultProfile); },
+				.GetValue = [this]()			{ return GetForceBuild(); },
+				.SetValue = [this](bool bValue)	{ SetForceBuild(bValue); },
+				.GetDefaultValue = [this]()		{ return GetForceBuild(DefaultProfile); },
 				.IsEnabled = ForCode,
 			}
 		)
@@ -383,6 +367,40 @@ void FCustomProfileMarshaller::Construct()
 
 
 
+	TreeData->AddHeading( LOCTEXT("DirectoryHeading", "Directory") )
+
+		// unlikely to want to change this in most use cases
+		//.AddDirectoryString( LOCTEXT("CustomStagingPathLabel", "Custom Stage Directory"),
+		//	{
+		//		.GetValue = [this]()				{ return Profile->GetPackageDirectory(); },
+		//		.SetValue = [this](FString Value)	{ Profile->SetPackageDirectory(Value); },
+		//		.GetDefaultValue = EmptyString,
+		//		.IsVisible = ForCooked,
+		//	}
+		//)
+
+		.AddBoolean( LOCTEXT("ArchiveBuildLabel", "Archive Build"),
+			{
+				.GetValue = [this]()			{ return Profile->IsArchiving(); },
+				.SetValue = [this](bool bValue)	{ Profile->SetArchive(bValue); },
+				.GetDefaultValue = [this]()		{ return DefaultProfile->IsArchiving(); },
+				.IsVisible = ForCooked,
+			}
+		)
+
+		.AddDirectoryString( LOCTEXT("ArchivePathLabel", "Archive Directory"),
+			{
+				.GetValue = [this]()				{ return Profile->GetArchiveDirectory(); },
+				.SetValue = [this](FString Value)	{ Profile->SetArchiveDirectory(Value); },
+				.GetDefaultValue = EmptyString,
+				.IsVisible = ForCooked,				
+				.IsEnabled = [this]()				{ return Profile->IsArchiving(); }
+			}
+		)
+	;
+
+
+
 	TreeData->AddHeading( LOCTEXT("DeployAndRunHeading", "Deploy And Run") )
 
 		.AddBoolean( LOCTEXT("DeployLabel", "Deploy To Device"),
@@ -390,7 +408,7 @@ void FCustomProfileMarshaller::Construct()
 				.GetValue = [this]()			{ return GetDeployToDevice(); },
 				.SetValue = [this](bool bValue)	{ SetDeployToDevice(bValue); },
 				.GetDefaultValue = [this]()		{ return GetDeployToDevice(DefaultProfile); },
-				.IsEnabled = [this]() { return (ContentScheme != CustomLaunchUI::EContentScheme::CookOnTheFly); }
+				.IsVisible = [this]()			{ return (ContentScheme != CustomLaunchUI::EContentScheme::CookOnTheFly); }
 			}
 		)
 
@@ -452,6 +470,7 @@ void FCustomProfileMarshaller::Construct()
 
 	MapOption = Profile->GetCookedMaps().Num() > 0 ? EMapOption::Selected : EMapOption::Startup;
 	DeployDeviceOption = Profile->GetDeployedDeviceGroup()->GetDeviceIDs().Num() > 0 ? EDeployDeviceOption::Selected : EDeployDeviceOption::Default;
+	bShouldCook = GetCook();
 	ContentScheme = Model->DetermineProfileContentScheme(Profile);
 	CacheStartupMapList();
 }
@@ -628,10 +647,16 @@ EBuildConfiguration FCustomProfileMarshaller::GetBuildConfiguration() const
 
 
 
+void FCustomProfileMarshaller::RefreshContentScheme()
+{
+	CustomLaunchUI::EContentScheme CurrentContentScheme = Model->DetermineProfileContentScheme(Profile);
+	SetContentScheme(CurrentContentScheme);
+}
+
 void FCustomProfileMarshaller::SetContentScheme(CustomLaunchUI::EContentScheme InContentScheme)
 {
 	ELauncherProfileDeploymentModes::Type DeploymentMode = GetDeployToDevice() ? ELauncherProfileDeploymentModes::CopyToDevice : ELauncherProfileDeploymentModes::DoNotDeploy;
-	Model->SetProfileContentScheme(InContentScheme, Profile, DeploymentMode );
+	Model->SetProfileContentScheme(InContentScheme, Profile, bShouldCook, DeploymentMode );
 	OnPropertyChanged();
 
 	// refresh the cached content scheme in case the option that was selected is not available
@@ -771,6 +796,23 @@ TSharedRef<SWidget> FCustomProfileMarshaller::CreateMapListWidget()
 END_SLATE_FUNCTION_BUILD_OPTIMIZATION
 
 
+void FCustomProfileMarshaller::SetCook( bool bCook )
+{
+	bShouldCook = bCook;
+	RefreshContentScheme();
+}
+
+bool FCustomProfileMarshaller::GetCook(ILauncherProfilePtr InProfile) const
+{
+	if (InProfile == nullptr)
+	{
+		InProfile = Profile;
+	}
+	return (InProfile->GetCookMode() != ELauncherProfileCookModes::DoNotCook);
+}
+
+
+
 void FCustomProfileMarshaller::SetMapsToCook(TArray<FString> MapsToCook)
 {
 	Profile->ClearCookedMaps(); 
@@ -970,13 +1012,13 @@ void FCustomProfileMarshaller::SetDeployDeviceOption( EDeployDeviceOption NewDep
 }
 
 
-void FCustomProfileMarshaller::SetBuildCode( bool bBuildCode )
+void FCustomProfileMarshaller::SetBuild( bool bBuild )
 {
-	if (!bBuildCode)
+	if (!bBuild)
 	{
 		Profile->SetBuildMode(ELauncherProfileBuildModes::DoNotBuild);
 	}
-	else if (GetBuildAlways(Profile))
+	else if (GetForceBuild(Profile))
 	{
 		Profile->SetBuildMode(ELauncherProfileBuildModes::Build);
 	}
@@ -988,7 +1030,7 @@ void FCustomProfileMarshaller::SetBuildCode( bool bBuildCode )
 	OnPropertyChanged();
 }
 
-bool FCustomProfileMarshaller::GetBuildCode(ILauncherProfilePtr InProfile) const
+bool FCustomProfileMarshaller::GetBuild(ILauncherProfilePtr InProfile) const
 {
 	if (InProfile == nullptr)
 	{
@@ -997,13 +1039,13 @@ bool FCustomProfileMarshaller::GetBuildCode(ILauncherProfilePtr InProfile) const
 	return (InProfile->GetBuildMode() != ELauncherProfileBuildModes::DoNotBuild);
 }
 
-void FCustomProfileMarshaller::SetBuildAlways( bool bBuildAlways )
+void FCustomProfileMarshaller::SetForceBuild( bool bForceBuild )
 {
-	if (!GetBuildCode())
+	if (!GetBuild())
 	{
 		Profile->SetBuildMode(ELauncherProfileBuildModes::DoNotBuild);
 	}
-	else if (bBuildAlways)
+	else if (bForceBuild)
 	{
 		Profile->SetBuildMode(ELauncherProfileBuildModes::Build);
 	}
@@ -1015,7 +1057,7 @@ void FCustomProfileMarshaller::SetBuildAlways( bool bBuildAlways )
 	OnPropertyChanged();
 }
 
-bool FCustomProfileMarshaller::GetBuildAlways(ILauncherProfilePtr InProfile) const
+bool FCustomProfileMarshaller::GetForceBuild(ILauncherProfilePtr InProfile) const
 {
 	if (InProfile == nullptr)
 	{
@@ -1033,6 +1075,7 @@ void FCustomProfileMarshaller::SetDeployToDevice( bool bDeployToDevice )
 		OnPropertyChanged();
 	}
 }
+
 bool FCustomProfileMarshaller::GetDeployToDevice(ILauncherProfilePtr InProfile) const
 {
 	if (InProfile == nullptr)
diff --git a/Engine/Plugins/Developer/CustomLaunchUI/Source/Private/Model/CustomLaunchUIModel.cpp b/Engine/Plugins/Developer/CustomLaunchUI/Source/Private/Model/CustomLaunchUIModel.cpp
--- a/Engine/Plugins/Developer/CustomLaunchUI/Source/Private/Model/CustomLaunchUIModel.cpp
+++ b/Engine/Plugins/Developer/CustomLaunchUI/Source/Private/Model/CustomLaunchUIModel.cpp
@@ -484,7 +484,7 @@ namespace CustomLaunchUI
 
 
 
-	void FModel::SetProfileContentScheme(EContentScheme ContentScheme, const ILauncherProfileRef& Profile, ELauncherProfileDeploymentModes::Type DefaultDeploymentMode)
+	void FModel::SetProfileContentScheme(EContentScheme ContentScheme, const ILauncherProfileRef& Profile, bool bWantToCook, ELauncherProfileDeploymentModes::Type DefaultDeploymentMode)
 	{
 		bool bPakFiles = (ContentScheme == EContentScheme::PakFiles || ContentScheme == EContentScheme::DevelopmentPackage);
 		bool bUseZen = (ContentScheme != EContentScheme::LooseFiles && ContentScheme != EContentScheme::CookOnTheFly);
@@ -503,7 +503,7 @@ namespace CustomLaunchUI
 			Profile->SetCookMode(ELauncherProfileCookModes::OnTheFly);
 			DeploymentMode = ELauncherProfileDeploymentModes::FileServer;
 		}
-		else if (bZenPakStreaming)
+		else if (bZenPakStreaming || !bWantToCook)
 		{
 			Profile->SetCookMode(ELauncherProfileCookModes::DoNotCook);
 		}
diff --git a/Engine/Plugins/Developer/CustomLaunchUI/Source/Private/Model/CustomLaunchUIModel.h b/Engine/Plugins/Developer/CustomLaunchUI/Source/Private/Model/CustomLaunchUIModel.h
--- a/Engine/Plugins/Developer/CustomLaunchUI/Source/Private/Model/CustomLaunchUIModel.h
+++ b/Engine/Plugins/Developer/CustomLaunchUI/Source/Private/Model/CustomLaunchUIModel.h
@@ -134,7 +134,7 @@ namespace CustomLaunchUI
 		const ILauncherProfileRef GetDefaultCustomLaunchProfile() const;
 
 		EContentScheme DetermineProfileContentScheme(const ILauncherProfileRef& Profile) const;
-		void SetProfileContentScheme(EContentScheme ContentScheme, const ILauncherProfileRef& Profile, ELauncherProfileDeploymentModes::Type DefaultDeploymentMode = ELauncherProfileDeploymentModes::CopyToDevice);
+		void SetProfileContentScheme(EContentScheme ContentScheme, const ILauncherProfileRef& Profile, bool bWantToCook = true, ELauncherProfileDeploymentModes::Type DefaultDeploymentMode = ELauncherProfileDeploymentModes::CopyToDevice);
 
 		TSharedPtr<FLaunchLogMessage> AddLogMessage( const FString& InMessage, ELogVerbosity::Type InVerbosity );
 		void ClearLogMessages();
diff --git a/Engine/Plugins/Developer/CustomLaunchUI/Source/Private/Widgets/Profiles/SCustomLaunchCustomProfileEditor.cpp b/Engine/Plugins/Developer/CustomLaunchUI/Source/Private/Widgets/Profiles/SCustomLaunchCustomProfileEditor.cpp
--- a/Engine/Plugins/Developer/CustomLaunchUI/Source/Private/Widgets/Profiles/SCustomLaunchCustomProfileEditor.cpp
+++ b/Engine/Plugins/Developer/CustomLaunchUI/Source/Private/Widgets/Profiles/SCustomLaunchCustomProfileEditor.cpp
@@ -48,11 +48,26 @@ public:
 
 	void Construct(const FArguments& InArgs, FLaunchProfileTreeNodePtr TreeNode, const TSharedRef<STableViewBase>& OwnerTable )
 	{
+		auto HasVisibleChildren = [TreeNode]()
+		{
+			for (FLaunchProfileTreeNodePtr ChildNode : TreeNode->Children)
+			{
+				if (ChildNode->Callbacks.IsVisible == nullptr || ChildNode->Callbacks.IsVisible())
+				{
+					return true;
+				}
+			}
+
+			return false;
+		};
+
+
 		this->ChildSlot
 		[
 			SNew(SBorder)
 			.BorderImage(FAppStyle::GetBrush("Brushes.Background"))
 			.Padding(FMargin(0, 0, 0, 1))
+			.Visibility_Lambda( [HasVisibleChildren]() { return HasVisibleChildren() ? EVisibility::Visible : EVisibility::Collapsed; } )
 			[
 				SNew(SBorder)
 				.BorderImage(FAppStyle::GetBrush("Brushes.Header"))				
diff --git a/Engine/Source/Developer/LauncherServices/Private/Launcher/LauncherWorker.cpp b/Engine/Source/Developer/LauncherServices/Private/Launcher/LauncherWorker.cpp
--- a/Engine/Source/Developer/LauncherServices/Private/Launcher/LauncherWorker.cpp
+++ b/Engine/Source/Developer/LauncherServices/Private/Launcher/LauncherWorker.cpp
@@ -898,7 +898,7 @@ FString FLauncherWorker::CreateUATCommand( const ILauncherProfileRef& InProfile,
 		UATCommand += TEXT(" -distribution");
 	}
 
-	if (InProfile->IsCookingIncrementally())
+	if (InProfile->IsCookingIncrementally() && !InProfile->IsImportingZenSnapshot()) // don't set -iterativecooking when importing a zen snapshot because it sets -cookincremental (the new system) instead.
 	{
 		UATCommand += TEXT(" -iterativecooking");
 	}
