Fix unreachable code warnings

#rnx

diff --git a/Engine/Plugins/Animation/ControlRig/Source/ControlRigEditor/Private/EditMode/SControlRigOutliner.cpp b/Engine/Plugins/Animation/ControlRig/Source/ControlRigEditor/Private/EditMode/SControlRigOutliner.cpp
--- a/Engine/Plugins/Animation/ControlRig/Source/ControlRigEditor/Private/EditMode/SControlRigOutliner.cpp
+++ b/Engine/Plugins/Animation/ControlRig/Source/ControlRigEditor/Private/EditMode/SControlRigOutliner.cpp
@@ -146,7 +146,6 @@ FText FMultiRigData::GetToolTipText(const FRigTreeDisplaySettings& Settings) con
 	{
 		return FText::Format(LOCTEXT("SMultiRigHierarchyModuleControlTooltip", "{0}\n\nUse Alt+Click to select subtree."), GetDisplayName(Settings));
 	}
-	return FText();
 }
 
 bool FMultiRigData::operator== (const FMultiRigData & Other) const
diff --git a/Engine/Plugins/Developer/CustomLaunchUI/Source/Private/Model/CustomLaunchUIModel.cpp b/Engine/Plugins/Developer/CustomLaunchUI/Source/Private/Model/CustomLaunchUIModel.cpp
--- a/Engine/Plugins/Developer/CustomLaunchUI/Source/Private/Model/CustomLaunchUIModel.cpp
+++ b/Engine/Plugins/Developer/CustomLaunchUI/Source/Private/Model/CustomLaunchUIModel.cpp
@@ -645,12 +645,12 @@ namespace CustomLaunchUI
 	{
 		switch (ContentScheme)
 		{
-			case EContentScheme::PakFiles:           return LOCTEXT("ContentSchemePakFiles","Pak Files"); break;
-			case EContentScheme::ZenStreaming:       return LOCTEXT("ContentSchemeZenStreaming","Zen Streaming"); break;
-			case EContentScheme::ZenPakStreaming:    return LOCTEXT("ContentSchemeZenPakStreaming","Zen Pak Streaming"); break;
-			case EContentScheme::DevelopmentPackage: return LOCTEXT("ContentSchemeDevPackage","Development Package"); break;
-			case EContentScheme::LooseFiles:         return LOCTEXT("ContentSchemeLooseFiles","Loose Files (legacy)"); break;
-			case EContentScheme::CookOnTheFly:       return LOCTEXT("ContentSchemeCOTF","Cook On The Fly"); break;
+			case EContentScheme::PakFiles:           return LOCTEXT("ContentSchemePakFiles","Pak Files"); 
+			case EContentScheme::ZenStreaming:       return LOCTEXT("ContentSchemeZenStreaming","Zen Streaming");
+			case EContentScheme::ZenPakStreaming:    return LOCTEXT("ContentSchemeZenPakStreaming","Zen Pak Streaming"); 
+			case EContentScheme::DevelopmentPackage: return LOCTEXT("ContentSchemeDevPackage","Development Package");
+			case EContentScheme::LooseFiles:         return LOCTEXT("ContentSchemeLooseFiles","Loose Files (legacy)");
+			case EContentScheme::CookOnTheFly:       return LOCTEXT("ContentSchemeCOTF","Cook On The Fly"); 
 		}
 
 		checkNoEntry();
@@ -662,12 +662,12 @@ namespace CustomLaunchUI
 	{
 		switch (ContentScheme)
 		{
-			case EContentScheme::PakFiles:           return LOCTEXT("ContentSchemeTipPakFiles","Store cooked game content in one ore more large Pak Files"); break;
-			case EContentScheme::ZenStreaming:       return LOCTEXT("ContentSchemeTipZenStreaming","Stream cooked game content from Zen Server"); break;
-			case EContentScheme::ZenPakStreaming:    return LOCTEXT("ContentSchemeTipZenPakStreaming","Stream an existing Pak Files build via a Zen"); break;
-			case EContentScheme::DevelopmentPackage: return LOCTEXT("ContentSchemeTipDevPackage","Package cooked game content into a single installable package file for development purposes, where available"); break;
-			case EContentScheme::LooseFiles:         return LOCTEXT("ContentSchemeTipLooseFiles","Store cooked game assets in individual files (legacy - recommend moving to Zen Streaming. This option will not work if the project is already configured with 'Use Zen Store')"); break;
-			case EContentScheme::CookOnTheFly:       return LOCTEXT("ContentSchemeTipCOTF","Only cook game assets when the game requires them, and send them over the network (legacy - slow)"); break;
+			case EContentScheme::PakFiles:           return LOCTEXT("ContentSchemeTipPakFiles","Store cooked game content in one ore more large Pak Files");
+			case EContentScheme::ZenStreaming:       return LOCTEXT("ContentSchemeTipZenStreaming","Stream cooked game content from Zen Server");
+			case EContentScheme::ZenPakStreaming:    return LOCTEXT("ContentSchemeTipZenPakStreaming","Stream an existing Pak Files build via a Zen");
+			case EContentScheme::DevelopmentPackage: return LOCTEXT("ContentSchemeTipDevPackage","Package cooked game content into a single installable package file for development purposes, where available");
+			case EContentScheme::LooseFiles:         return LOCTEXT("ContentSchemeTipLooseFiles","Store cooked game assets in individual files (legacy - recommend moving to Zen Streaming. This option will not work if the project is already configured with 'Use Zen Store')");
+			case EContentScheme::CookOnTheFly:       return LOCTEXT("ContentSchemeTipCOTF","Only cook game assets when the game requires them, and send them over the network (legacy - slow)");
 		}
 
 		checkNoEntry();
diff --git a/Engine/Plugins/Editor/ModelingToolsEditorMode/Source/ModelingToolsEditorMode/Private/ModelingModeToolkit_Toolbars.cpp b/Engine/Plugins/Editor/ModelingToolsEditorMode/Source/ModelingToolsEditorMode/Private/ModelingModeToolkit_Toolbars.cpp
--- a/Engine/Plugins/Editor/ModelingToolsEditorMode/Source/ModelingToolsEditorMode/Private/ModelingModeToolkit_Toolbars.cpp
+++ b/Engine/Plugins/Editor/ModelingToolsEditorMode/Source/ModelingToolsEditorMode/Private/ModelingModeToolkit_Toolbars.cpp
@@ -513,7 +513,7 @@ TSharedRef<SWidget> FModelingToolsEditorModeToolkit::MakeMenu_ModelingModeConfig
 void FModelingToolsEditorModeToolkit::ExtendSecondaryModeToolbar(UToolMenu *InModeToolbarMenu)
 {
 	return;		// disable for now
-
+#if 0
 	/*
 	 * Sanity check because the toolbar extension should happen after the ModeUILayer setup which is currently manual
 	 * for standalone asset editors
@@ -618,7 +618,7 @@ void FModelingToolsEditorModeToolkit::ExtendSecondaryModeToolbar(UToolMenu *InMo
 	Section.AddEntry(FToolMenuEntry::InitWidget("DefaultSettings", DefaultSettingsWidget, FText()));
 	
 	Section.AddSeparator(NAME_None);
-
+#endif
 }
 
 
diff --git a/Engine/Plugins/MassInsights/Source/MassInsightsUI/Private/MassInsightsUI/Widgets/SEntityEventsTableView.cpp b/Engine/Plugins/MassInsights/Source/MassInsightsUI/Private/MassInsightsUI/Widgets/SEntityEventsTableView.cpp
--- a/Engine/Plugins/MassInsights/Source/MassInsightsUI/Private/MassInsightsUI/Widgets/SEntityEventsTableView.cpp
+++ b/Engine/Plugins/MassInsights/Source/MassInsightsUI/Private/MassInsightsUI/Widgets/SEntityEventsTableView.cpp
@@ -97,13 +97,10 @@ namespace MassInsightsUI
 				{
 				case MassInsightsAnalysis::EMassEntityEventType::Created:
 					return FAppStyle::GetBrush("Icons.Plus");
-					break;
 				case MassInsightsAnalysis::EMassEntityEventType::ArchetypeChange:
 					return FAppStyle::GetBrush("Icons.ArrowRight");
-					break;
 				case MassInsightsAnalysis::EMassEntityEventType::Destroyed:
 					return FAppStyle::GetBrush("Icons.X");
-					break;
 				default:
 					return FAppStyle::GetBrush("Icons.Warning");
 				}
diff --git a/Engine/Source/Editor/Experimental/EditorInteractiveToolsFramework/Private/EditorDragTools/EditorViewportChange.cpp b/Engine/Source/Editor/Experimental/EditorInteractiveToolsFramework/Private/EditorDragTools/EditorViewportChange.cpp
--- a/Engine/Source/Editor/Experimental/EditorInteractiveToolsFramework/Private/EditorDragTools/EditorViewportChange.cpp
+++ b/Engine/Source/Editor/Experimental/EditorInteractiveToolsFramework/Private/EditorDragTools/EditorViewportChange.cpp
@@ -173,28 +173,20 @@ FText FEditorViewportChange::GetDesiredViewportTypeText() const
 	{
 	case LVT_Perspective:
 		return FText::FromString("Perspective");
-		break;
 	case LVT_OrthoFreelook:
 		return FText::FromString("Free Look");
-		break;
 	case LVT_OrthoTop:
 		return FText::FromString("Top");
-		break;
 	case LVT_OrthoLeft:
 		return FText::FromString("Left");
-		break;
 	case LVT_OrthoFront:
 		return FText::FromString("Front");
-		break;
 	case LVT_OrthoBack:
 		return FText::FromString("Back");
-		break;
 	case LVT_OrthoBottom:
 		return FText::FromString("Bottom");
-		break;
 	case LVT_OrthoRight:
 		return FText::FromString("Right");
-		break;
 	default:;
 	}
 
diff --git a/Engine/Source/Editor/SequencerCore/Private/MVVM/Extensions/IDeactivatableExtension.cpp b/Engine/Source/Editor/SequencerCore/Private/MVVM/Extensions/IDeactivatableExtension.cpp
--- a/Engine/Source/Editor/SequencerCore/Private/MVVM/Extensions/IDeactivatableExtension.cpp
+++ b/Engine/Source/Editor/SequencerCore/Private/MVVM/Extensions/IDeactivatableExtension.cpp
@@ -77,9 +77,9 @@ ECachedDeactiveState FDeactiveStateCacheExtension::ComputeFlagsForModel(const FV
 		{
 			ThisModelFlags |= ECachedDeactiveState::Inheritable;
 		}
-	}*/
+	}
 
-	return ThisModelFlags;
+	return ThisModelFlags;*/
 }
 
 void FDeactiveStateCacheExtension::PostComputeChildrenFlags(const FViewModelPtr& ViewModel, ECachedDeactiveState& OutThisModelFlags, ECachedDeactiveState& OutPropagateToParentFlags)
