Updates for BuildCookRunUI:
- add reusable combo box control & use it to replace build configuration combo
- add architecture selection

#rnx

diff --git a/Engine/Plugins/Developer/CustomLaunchUI/Source/Private/Marshalling/BasicProfileMarshaller.cpp b/Engine/Plugins/Developer/CustomLaunchUI/Source/Private/Marshalling/BasicProfileMarshaller.cpp
--- a/Engine/Plugins/Developer/CustomLaunchUI/Source/Private/Marshalling/BasicProfileMarshaller.cpp
+++ b/Engine/Plugins/Developer/CustomLaunchUI/Source/Private/Marshalling/BasicProfileMarshaller.cpp
@@ -4,8 +4,8 @@
 #include "Model/CustomLaunchUIModel.h"
 #include "Widgets/Shared/SCustomLaunchProjectCombo.h"
 #include "Widgets/Shared/SCustomLaunchBuildTargetCombo.h"
-#include "Widgets/Shared/SCustomLaunchBuildConfigurationCombo.h"
 #include "Widgets/Shared/SCustomLaunchContentSchemeCombo.h"
+#include "Widgets/Shared/SCustomLaunchCombo.h"
 #include "Widgets/Shared/SCustomLaunchDeviceListView.h"
 #include "Widgets/Input/SMultiLineEditableTextBox.h"
 #include "SlateOptMacros.h"
@@ -99,10 +99,10 @@ void FBasicProfileMarshaller::Construct()
 		)
 
 		.AddWidget( LOCTEXT("ConfigurationLabel", "Configuration"), 
-			SNew(SCustomLaunchBuildConfigurationCombo)
+			SNew(SCustomLaunchLexToStringCombo<EBuildConfiguration>)
 			.OnSelectionChanged(this, &FBasicProfileMarshaller::SetBuildConfiguration)
-			.SelectedBuildConfiguration(this, &FBasicProfileMarshaller::GetBuildConfiguration)
-			.Font(FCoreStyle::Get().GetFontStyle("SmallFont"))
+			.SelectedItem(this, &FBasicProfileMarshaller::GetBuildConfiguration)
+			.Items( TArray<EBuildConfiguration>({ EBuildConfiguration::Debug, EBuildConfiguration::Development, EBuildConfiguration::Test, EBuildConfiguration::Shipping }) )
 		)
 
 		.AddWidget( LOCTEXT("ContentSchemeLabel", "Content Scheme"), 
diff --git a/Engine/Plugins/Developer/CustomLaunchUI/Source/Private/Marshalling/CustomProfileMarshaller.cpp b/Engine/Plugins/Developer/CustomLaunchUI/Source/Private/Marshalling/CustomProfileMarshaller.cpp
--- a/Engine/Plugins/Developer/CustomLaunchUI/Source/Private/Marshalling/CustomProfileMarshaller.cpp
+++ b/Engine/Plugins/Developer/CustomLaunchUI/Source/Private/Marshalling/CustomProfileMarshaller.cpp
@@ -5,10 +5,10 @@
 #include "Widgets/Shared/SCustomLaunchProjectCombo.h"
 #include "Widgets/Shared/SCustomLaunchBuildTargetCombo.h"
 #include "Widgets/Shared/SCustomLaunchPlatformCombo.h"
-#include "Widgets/Shared/SCustomLaunchBuildConfigurationCombo.h"
 #include "Widgets/Shared/SCustomLaunchContentSchemeCombo.h"
 #include "Widgets/Shared/SCustomLaunchDeviceListView.h"
 #include "Widgets/Shared/SCustomLaunchMapListView.h"
+#include "Widgets/Shared/SCustomLaunchCombo.h"
 #include "Widgets/Input/SSegmentedControl.h"
 #include "Widgets/Text/STextBlock.h"
 #include "Widgets/Layout/SBox.h"
@@ -18,6 +18,8 @@
 #include "SlateOptMacros.h"
 #include "PlatformInfo.h"
 #include "IDesktopPlatform.h"
+#include "Interfaces/ITargetPlatformManagerModule.h"
+#include "Interfaces/ITargetPlatform.h"
 #include "GameProjectHelper.h"
 
 
@@ -109,6 +111,10 @@ protected:
 	void SetForceBuild( bool bForceBuild );
 	bool GetForceBuild( ILauncherProfilePtr InProfile = nullptr ) const;
 
+	void SetArchitecture( FString Architecture );
+	FString GetArchitecture() const;
+	FText GetArchitectureDisplayName( FString Architecture );
+
 	void SetDeployToDevice( bool bDeployToDevice );
 	bool GetDeployToDevice( ILauncherProfilePtr InProfile = nullptr ) const;
 
@@ -121,9 +127,11 @@ protected:
 
 
 protected:
+
 	virtual void OnPropertyChanged() override;
 	void RefreshContentScheme();
 	void CacheStartupMapList();
+	void CacheArchitectures();
 	TSharedRef<SWidget> OnGenerateComboWidget( TSharedPtr<FString> InComboString );
 
 
@@ -138,6 +146,8 @@ protected:
 	TSharedPtr<SSearchableComboBox> InitalMapCombo;
 	TSharedPtr<SCustomLaunchMapListView> MapListView;
 	TSharedPtr<SCustomLaunchDeviceListView> DeployDeviceListView;
+	TArray<FString> CachedArchitectures;
+	EBuildTargetType CachedBuildTargetType = EBuildTargetType::Game;
 };
 
 
@@ -229,10 +239,10 @@ void FCustomProfileMarshaller::Construct()
 		)
 
 		.AddWidget( LOCTEXT("ConfigurationLabel", "Configuration"), 
-			SNew(SCustomLaunchBuildConfigurationCombo)
+			SNew(SCustomLaunchLexToStringCombo<EBuildConfiguration>)
 			.OnSelectionChanged(this, &FCustomProfileMarshaller::SetBuildConfiguration)
-			.SelectedBuildConfiguration(this, &FCustomProfileMarshaller::GetBuildConfiguration)
-			.Font(FCoreStyle::Get().GetFontStyle("SmallFont"))
+			.SelectedItem(this, &FCustomProfileMarshaller::GetBuildConfiguration)
+			.Items( TArray<EBuildConfiguration>({ EBuildConfiguration::Debug, EBuildConfiguration::Development, EBuildConfiguration::Test, EBuildConfiguration::Shipping }) )
 		)
 
 		.AddWidget( LOCTEXT("ContentSchemeLabel", "Content Scheme"), 
@@ -362,7 +372,18 @@ void FCustomProfileMarshaller::Construct()
 		//	}
 		//)
 
-
+		.AddWidget( LOCTEXT("ArchitectureLabel", "Architecture"),
+			{
+				.IsDefault = [this]()		{ return GetArchitecture().IsEmpty(); },
+				.SetToDefault = [this]()	{ SetArchitecture(FString()); },
+				.IsVisible = [this]()       { return CachedArchitectures.Num() > 0; }
+			},
+			SNew(SCustomLaunchStringCombo)
+			.OnSelectionChanged( this, &FCustomProfileMarshaller::SetArchitecture )
+			.SelectedItem( this, &FCustomProfileMarshaller::GetArchitecture)
+			.GetDisplayName( this, &FCustomProfileMarshaller::GetArchitectureDisplayName )
+			.Items_Lambda( [this]() { return CachedArchitectures;} )
+		)
 	;
 
 
@@ -472,7 +493,10 @@ void FCustomProfileMarshaller::Construct()
 	DeployDeviceOption = Profile->GetDeployedDeviceGroup()->GetDeviceIDs().Num() > 0 ? EDeployDeviceOption::Selected : EDeployDeviceOption::Default;
 	bShouldCook = GetCook();
 	ContentScheme = Model->DetermineProfileContentScheme(Profile);
+	CachedBuildTargetType = Model->GetBuildTargetInfo(Profile).Type;
+	
 	CacheStartupMapList();
+	CacheArchitectures();
 }
 END_SLATE_FUNCTION_BUILD_OPTIMIZATION
 
@@ -497,6 +521,27 @@ void FCustomProfileMarshaller::CacheStartupMapList()
 	InitalMapCombo->RefreshOptions();
 }
 
+void FCustomProfileMarshaller::CacheArchitectures()
+{
+	CachedArchitectures.Reset();
+
+	const PlatformInfo::FTargetPlatformInfo* PlatformInfo = Model->GetPlatformInfo(Profile);
+	if (PlatformInfo != nullptr)
+	{
+		const ITargetPlatform* TargetPlatform = GetTargetPlatformManager()->FindTargetPlatform(PlatformInfo->IniPlatformName);
+		if (TargetPlatform)
+		{
+			TargetPlatform->GetPossibleArchitectures(CachedArchitectures);
+			if (CachedArchitectures.Num() > 0)
+			{
+				CachedArchitectures.Insert(FString(), 0); // empty string for "project default" option
+			}
+		}
+	}
+}
+
+
+
 TSharedRef<SWidget> FCustomProfileMarshaller::OnGenerateComboWidget( TSharedPtr<FString> InComboString )
 {
 	return SNew(STextBlock)
@@ -530,6 +575,8 @@ void FCustomProfileMarshaller::SetSelectedPlatforms( TArray<FString> SelectedPla
 	OnPropertyChanged();
 
 	DeployDeviceListView->OnSelectedPlatformChanged();
+
+	CacheArchitectures();
 }
 
 TArray<FString> FCustomProfileMarshaller::GetSelectedPlatforms() const
@@ -603,6 +650,7 @@ void FCustomProfileMarshaller::SetBuildTarget(FString BuildTarget)
 		Model->UpdateCookedPlatformsFromBuildTarget(Profile);
 	}
 
+	CachedBuildTargetType = Model->GetBuildTargetInfo(Profile).Type;
 	OnPropertyChanged();
 }
 
@@ -1066,6 +1114,62 @@ bool FCustomProfileMarshaller::GetForceBuild(ILauncherProfilePtr InProfile) cons
 	return (InProfile->GetBuildMode() == ELauncherProfileBuildModes::Build);
 }
 
+void FCustomProfileMarshaller::SetArchitecture( FString Architecture )
+{
+	// clear existing architectures
+	TArray<FString> Empty;
+	Profile->SetServerArchitectures(Empty);
+	Profile->SetEditorArchitectures(Empty);
+	Profile->SetClientArchitectures(Empty);
+
+	// set new single architecture
+	if (!Architecture.IsEmpty())
+	{
+		TArray<FString> Single;
+		Single.Add(Architecture);
+
+		switch (CachedBuildTargetType)
+		{
+			case EBuildTargetType::Server: Profile->SetServerArchitectures(Single); break;
+			case EBuildTargetType::Editor: Profile->SetEditorArchitectures(Single); break;
+			default:                       Profile->SetClientArchitectures(Single); break;
+		}
+	}
+
+	OnPropertyChanged();
+}
+
+FString FCustomProfileMarshaller::GetArchitecture() const
+{
+	TArray<FString> Architectures;
+
+	switch (CachedBuildTargetType)
+	{
+		case EBuildTargetType::Server: Architectures = Profile->GetServerArchitectures(); break;
+		case EBuildTargetType::Editor: Architectures = Profile->GetEditorArchitectures(); break;
+		default:                       Architectures = Profile->GetClientArchitectures(); break;
+	}
+
+	return (Architectures.Num() > 0) ? Architectures[0] : FString();
+}
+
+FText FCustomProfileMarshaller::GetArchitectureDisplayName( FString Architecture )
+{
+	if (Architecture.IsEmpty())
+	{
+		return LOCTEXT("DefaultArchName", "Project Default");
+	}
+	else if (Architecture == FPlatformMisc::GetHostArchitecture() && Model->IsHostPlatform(Profile))
+	{
+		return FText::Format(LOCTEXT("HostArchLabel", "{0} (this platform)"), FText::FromString(Architecture) );
+	}
+	else
+	{
+		return FText::FromString(Architecture);
+	}
+}
+
+
 
 void FCustomProfileMarshaller::SetDeployToDevice( bool bDeployToDevice )
 {
diff --git a/Engine/Plugins/Developer/CustomLaunchUI/Source/Private/Widgets/Shared/SCustomLaunchBuildConfigurationCombo.cpp b/Engine/Plugins/Developer/CustomLaunchUI/Source/Private/Widgets/Shared/SCustomLaunchBuildConfigurationCombo.cpp
deleted file mode 100644
diff --git a/Engine/Plugins/Developer/CustomLaunchUI/Source/Private/Widgets/Shared/SCustomLaunchBuildConfigurationCombo.h b/Engine/Plugins/Developer/CustomLaunchUI/Source/Private/Widgets/Shared/SCustomLaunchBuildConfigurationCombo.h
deleted file mode 100644
diff --git a/Engine/Plugins/Developer/CustomLaunchUI/Source/Private/Widgets/Shared/SCustomLaunchCombo.h b/Engine/Plugins/Developer/CustomLaunchUI/Source/Private/Widgets/Shared/SCustomLaunchCombo.h
new file mode 100644
--- /dev/null
+++ b/Engine/Plugins/Developer/CustomLaunchUI/Source/Private/Widgets/Shared/SCustomLaunchCombo.h
@@ -0,0 +1,144 @@
+// Copyright Epic Games, Inc. All Rights Reserved.
+
+#pragma once
+
+#include "CoreMinimal.h"
+#include "SlateOptMacros.h"
+#include "Widgets/SCompoundWidget.h"
+#include "Widgets/Input/SComboButton.h"
+#include "Styling/SlateTypes.h"
+#include "Styling/CoreStyle.h"
+#include "Framework/MultiBox/MultiBoxBuilder.h"
+
+// helper class to create a combo box that can select an item
+template<typename T>
+class SCustomLaunchCombo
+	: public SCompoundWidget
+{
+public:
+	DECLARE_DELEGATE_OneParam(FOnSelectionChanged, T );
+	DECLARE_DELEGATE_RetVal_OneParam(FText, FGetItemText, T)
+
+	SLATE_BEGIN_ARGS(SCustomLaunchCombo<T>)
+		{}
+		SLATE_EVENT(FOnSelectionChanged, OnSelectionChanged)
+		SLATE_EVENT(FGetItemText, GetDisplayName)
+		SLATE_EVENT(FGetItemText, GetItemToolTip)
+		SLATE_ATTRIBUTE(TArray<T>, Items)
+		SLATE_ATTRIBUTE(T, SelectedItem)
+	SLATE_END_ARGS()
+
+public:
+	void Construct(	const FArguments& InArgs)
+	{
+		OnSelectionChanged = InArgs._OnSelectionChanged;
+		GetDisplayName = InArgs._GetDisplayName;
+		GetItemToolTip = InArgs._GetItemToolTip;
+		Items = InArgs._Items;
+		SelectedItem = InArgs._SelectedItem;
+
+		ChildSlot
+		[
+			SNew(SComboButton)
+			.ButtonContent()
+			[
+				SNew(STextBlock)
+				.Text(this, &SCustomLaunchCombo<T>::GetSelectedItemDisplayName)
+				.Font(FCoreStyle::Get().GetFontStyle("SmallFont"))
+			]
+			.OnGetMenuContent(this, &SCustomLaunchCombo<T>::MakeWidget)
+		];
+	}
+
+protected:
+	TAttribute<T> SelectedItem;
+	TAttribute<TArray<T>> Items;
+	FOnSelectionChanged OnSelectionChanged;
+	FGetItemText GetDisplayName;
+	FGetItemText GetItemToolTip;
+
+	FText GetSelectedItemDisplayName() const
+	{
+		T Item = SelectedItem.Get();
+		return ToText(Item);
+	}
+
+	void SetSelectedItem(T Value)
+	{
+		OnSelectionChanged.ExecuteIfBound(Value);
+	}
+
+	FText ToText(T Value) const
+	{
+		if (GetDisplayName.IsBound())
+		{
+			return GetDisplayName.Execute(Value);
+		}
+
+		return ToTextOverride(Value);
+	}
+
+
+	virtual FText ToTextOverride(T Value) const
+	{
+		return FText::FromString("ERROR: GetDisplayName or ToTextOverride not bound");
+	}
+
+
+BEGIN_SLATE_FUNCTION_BUILD_OPTIMIZATION
+	TSharedRef<SWidget> MakeWidget()
+	{
+		FMenuBuilder MenuBuilder(true, nullptr);
+		{
+			TArray<T> MenuItems = Items.Get();
+
+			for (const T& Item : MenuItems)
+			{
+				FText ItemText = ToText(Item);
+				FText ToolTipText = GetItemToolTip.IsBound() ? GetItemToolTip.Execute(Item) : FText::GetEmpty();
+
+				MenuBuilder.AddMenuEntry(
+					ItemText, 
+					ToolTipText, 
+					FSlateIcon(), 
+					FUIAction(
+						FExecuteAction::CreateSP(this, &SCustomLaunchCombo<T>::SetSelectedItem, Item),
+						FCanExecuteAction(),
+						FGetActionCheckState::CreateLambda( [this, Item]() { return (SelectedItem.Get() == Item) ? ECheckBoxState::Checked : ECheckBoxState::Unchecked; })
+					),
+					NAME_None,
+					EUserInterfaceActionType::Check);
+			}
+		}
+
+		return MenuBuilder.MakeWidget();
+	}
+END_SLATE_FUNCTION_BUILD_OPTIMIZATION
+
+
+};
+
+
+// helper class that can display a list of items, if those items have a LexToString function available
+template<typename T>
+class SCustomLaunchLexToStringCombo : public SCustomLaunchCombo<T>
+{
+public:
+	virtual FText ToTextOverride(T Value) const override
+	{
+		return FText::FromString(LexToString(Value));
+	}
+};
+
+// helper class that can display a list of strings
+class SCustomLaunchStringCombo : public SCustomLaunchCombo<FString>
+{
+public:
+	virtual FText ToTextOverride(FString Value) const override
+	{
+		return FText::FromString(Value);
+	}
+};
+
+
+
