Groundwork for CustomLaunchUI extensibility:
- Move marshalling system into public headers
- Add a Generic marshalling class with helper functions for all existing fields
- Migrate built-in Basic & Custom marshallers to extend Generic marshaller

#rnx

diff --git a/Engine/Plugins/Developer/CustomLaunchUI/Source/Private/CustomLaunchUIModule.cpp b/Engine/Plugins/Developer/CustomLaunchUI/Source/Private/CustomLaunchUIModule.cpp
--- a/Engine/Plugins/Developer/CustomLaunchUI/Source/Private/CustomLaunchUIModule.cpp
+++ b/Engine/Plugins/Developer/CustomLaunchUI/Source/Private/CustomLaunchUIModule.cpp
@@ -1,6 +1,6 @@
 // Copyright Epic Games, Inc. All Rights Reserved.
 
-#include "Modules/ModuleManager.h"
+#include "CustomLaunchUIModule.h"
 #include "WorkspaceMenuStructure.h"
 #include "WorkspaceMenuStructureModule.h"
 #include "Framework/Docking/TabManager.h"
@@ -18,15 +18,22 @@
 
 static const FName TabName("CustomLaunchUI");
 
-class FCustomLaunchUIModule : public IModuleInterface
+namespace CustomLaunchUI
+{
+	extern void RegisterMarshallerFactory( TSharedRef<ILaunchProfileMarshallerFactory> MarshallerFactory );
+	extern void UnregisterMarshallerFactory( TSharedRef<ILaunchProfileMarshallerFactory> MarshallerFactory );
+}
+
+
+class FCustomLaunchUIModule : public ICustomLaunchUIModule
 {
 public:
 
 	virtual void StartupModule() override
 	{
-		// register marshalls
-		ILaunchProfileMarshallerFactory::Register( MakeShared<FBasicProfileMarshallerFactory>() );
-		ILaunchProfileMarshallerFactory::Register( MakeShared<FCustomProfileMarshallerFactory>() );
+		// register default marshalls
+		RegisterMarshaller( MakeShared<CustomLaunchUI::FBasicProfileMarshallerFactory>() );
+		RegisterMarshaller( MakeShared<CustomLaunchUI::FCustomProfileMarshallerFactory>() );
 
 		// register custom styles
 		FCustomLaunchUIStyle::Initialize();
@@ -46,6 +53,19 @@ public:
 		FCustomLaunchUIStyle::Shutdown();
 	}
 
+	virtual void RegisterMarshaller( TSharedRef<CustomLaunchUI::ILaunchProfileMarshallerFactory> MarshallerFactory ) override
+	{
+		CustomLaunchUI::RegisterMarshallerFactory(MarshallerFactory);
+	}
+
+	virtual void UnregisterMarshaller( TSharedRef<CustomLaunchUI::ILaunchProfileMarshallerFactory> MarshallerFactory ) override
+	{
+		CustomLaunchUI::UnregisterMarshallerFactory(MarshallerFactory);
+	}
+
+
+
+
 private:
 
 	TSharedRef<SDockTab> HandleSpawnTab(const FSpawnTabArgs& SpawnTabArgs)
diff --git a/Engine/Plugins/Developer/CustomLaunchUI/Source/Private/Marshalling/BasicProfileMarshaller.cpp b/Engine/Plugins/Developer/CustomLaunchUI/Source/Private/Marshalling/BasicProfileMarshaller.cpp
--- a/Engine/Plugins/Developer/CustomLaunchUI/Source/Private/Marshalling/BasicProfileMarshaller.cpp
+++ b/Engine/Plugins/Developer/CustomLaunchUI/Source/Private/Marshalling/BasicProfileMarshaller.cpp
@@ -1,297 +1,38 @@
 // Copyright Epic Games, Inc. All Rights Reserved.
 
 #include "Marshalling/BasicProfileMarshaller.h"
-#include "Model/CustomLaunchUIModel.h"
-#include "Widgets/Shared/SCustomLaunchProjectCombo.h"
-#include "Widgets/Shared/SCustomLaunchBuildTargetCombo.h"
-#include "Widgets/Shared/SCustomLaunchContentSchemeCombo.h"
-#include "Widgets/Shared/SCustomLaunchCombo.h"
-#include "Widgets/Shared/SCustomLaunchDeviceListView.h"
-#include "Widgets/Input/SMultiLineEditableTextBox.h"
-#include "SlateOptMacros.h"
 
+#define LOCTEXT_NAMESPACE "BasicProfileMarshaller"
 
-#define LOCTEXT_NAMESPACE "CustomProfileMarshaller"
-
-class FBasicProfileMarshaller : public ILaunchProfileMarshaller, public TSharedFromThis<FBasicProfileMarshaller>
-{
-public:
-	FBasicProfileMarshaller( const ILauncherProfileRef& Profile, const TSharedRef<CustomLaunchUI::FModel>& InModel );
-	virtual ~FBasicProfileMarshaller() = default;
-
-	virtual void Construct() override;
-
-	virtual FString GetName() const override
-	{
-		return TEXT("BasicProfile");
-	}
-
-	virtual FLaunchProfileTreeDataRef GetProfileTree() override
-	{
-		return TreeData;
-	}
-
-protected:
-
-	FString GetProjectPath() const;
-	void SetProjectName(FString ProjectPath);
-	bool HasProject() const;
-
-	void SetBuildTarget(FString BuildTarget);
-	FString GetBuildTarget() const;
-
-	void SetBuildConfiguration(EBuildConfiguration BuildConfiguration);
-	EBuildConfiguration GetBuildConfiguration() const;
-
-	void SetContentScheme(CustomLaunchUI::EContentScheme ContentScheme);
-	bool IsContentSchemeAvailable(CustomLaunchUI::EContentScheme, FText& OutReason ) const;
-
-	FText GetCommandLine() const;
-	void SetCommandLine( const FText& NewText, ETextCommit::Type InTextCommit );
-
-	void SetDeployDeviceIDs(TArray<FString> DeployDeviceIDs);
-	TArray<FString> GetDeployDeviceIDs() const;
-
-protected:
-	virtual void OnPropertyChanged() override;
-
-	FLaunchProfileTreeDataRef TreeData;
-	const ILauncherProfileRef Profile;
-	const ILauncherProfileRef DefaultProfile;
-	const TSharedRef<CustomLaunchUI::FModel> Model;
-	CustomLaunchUI::EContentScheme ContentScheme;
-};
-
-
-
-FBasicProfileMarshaller::FBasicProfileMarshaller( const ILauncherProfileRef& InProfile, const TSharedRef<CustomLaunchUI::FModel>& InModel )
-	: TreeData( MakeShared<FLaunchProfileTreeData>(InProfile, InModel, this) )
-	, Profile(InProfile)
-	, DefaultProfile(InModel->GetDefaultBasicLaunchProfile())
-	, Model(InModel)
-{
-}
-
-BEGIN_SLATE_FUNCTION_BUILD_OPTIMIZATION
-void FBasicProfileMarshaller::Construct()
-{
-	TreeData->AddHeading( LOCTEXT("GeneralSettingsHeading", "General Settings") )
-
-		.AddWidget( LOCTEXT("ProjectLabel", "Project"), 
-			SNew(SCustomLaunchProjectCombo)
-			.OnSelectionChanged(this, &FBasicProfileMarshaller::SetProjectName)
-			.SelectedProject(this, &FBasicProfileMarshaller::GetProjectPath)
-			.HasProject(this, &FBasicProfileMarshaller::HasProject)
-			.CurrentProjectOption(SCustomLaunchProjectCombo::ECurrentProjectOption::Empty)
-			.Font(FCoreStyle::Get().GetFontStyle("SmallFont"))
-		)
-
-		.AddWidget( LOCTEXT("TargetLabel", "Target"), 
-			{
-				.IsDefault = [this]()		{ return !Profile->HasBuildTargetSpecified() || Profile->GetBuildTarget().IsEmpty(); },
-				.SetToDefault = [this]()	{ SetBuildTarget(FString()); },
-			},
-			SNew(SCustomLaunchBuildTargetCombo)
-			.OnSelectionChanged(this, &FBasicProfileMarshaller::SetBuildTarget)
-			.SelectedBuildTarget(this, &FBasicProfileMarshaller::GetBuildTarget)
-			.SelectedProject(this, &FBasicProfileMarshaller::GetProjectPath)
-			.Font(FCoreStyle::Get().GetFontStyle("SmallFont"))
-		)
-
-		.AddWidget( LOCTEXT("ConfigurationLabel", "Configuration"), 
-			SNew(SCustomLaunchLexToStringCombo<EBuildConfiguration>)
-			.OnSelectionChanged(this, &FBasicProfileMarshaller::SetBuildConfiguration)
-			.SelectedItem(this, &FBasicProfileMarshaller::GetBuildConfiguration)
-			.Items( TArray<EBuildConfiguration>({ EBuildConfiguration::Debug, EBuildConfiguration::Development, EBuildConfiguration::Test, EBuildConfiguration::Shipping }) )
-		)
-
-		.AddWidget( LOCTEXT("ContentSchemeLabel", "Content Scheme"), 
-			SNew(SCustomLaunchContentSchemeCombo)
-			.OnSelectionChanged(this, &FBasicProfileMarshaller::SetContentScheme)
-			.SelectedContentScheme_Lambda( [this]() { return ContentScheme; } )
-			.IsContentSchemeAvailable(this, &FBasicProfileMarshaller::IsContentSchemeAvailable)
-			.Font(FCoreStyle::Get().GetFontStyle("SmallFont"))
-		)
-
-		.AddWidget( LOCTEXT("CommandLineLabel", "Command Line"), 
-			{
-				.IsDefault = [this]()		{ return GetCommandLine().IsEmpty(); },
-				.SetToDefault = [this]()	{ SetCommandLine(FText(), ETextCommit::Default); },
-			},
-			SNew(SMultiLineEditableTextBox)
-			.AllowMultiLine(false)
-			.AutoWrapText(true)
-			.Text(this, &FBasicProfileMarshaller::GetCommandLine)
-			.OnTextCommitted(this, &FBasicProfileMarshaller::SetCommandLine)
-		)
-
-		.AddWidget( LOCTEXT("TargetDeviceLabel", "Target Device"), 
-			SNew(SCustomLaunchDeviceListView)
-			.OnSelectionChanged(this, &FBasicProfileMarshaller::SetDeployDeviceIDs)
-			.SelectedDevices(this, &FBasicProfileMarshaller::GetDeployDeviceIDs)
-			.AllPlatforms(true)
-			.SingleSelect(true)
-		)
-	;
-
-	ContentScheme = Model->DetermineProfileContentScheme(Profile);
-}
-END_SLATE_FUNCTION_BUILD_OPTIMIZATION
-
-FString FBasicProfileMarshaller::GetProjectPath() const
-{
-	if (Profile->HasProjectSpecified())
-	{
-		return Profile->GetProjectPath();
-	}
-	else
-	{
-		return Model->GetProfileManager()->GetProjectPath();
-	}
-}
-
-void FBasicProfileMarshaller::SetProjectName(FString ProjectPath)
-{
-	Profile->SetProjectPath(ProjectPath);
-	Profile->SetProjectSpecified(!ProjectPath.IsEmpty());
-
-	OnPropertyChanged();
-}
-
-bool FBasicProfileMarshaller::HasProject() const
-{
-	return Profile->HasProjectSpecified();
-}
-
-FString FBasicProfileMarshaller::GetBuildTarget() const
-{
-	if (Profile->HasBuildTargetSpecified())
-	{
-		return Profile->GetBuildTarget();
-	}
-	else
-	{
-		return Model->GetProfileManager()->GetBuildTarget();
-	}
-}
-
-void FBasicProfileMarshaller::SetBuildTarget(FString BuildTarget)
-{
-	Profile->SetBuildTarget(BuildTarget);
-	Profile->SetBuildTargetSpecified(!BuildTarget.IsEmpty());
-
-	Model->UpdatedCookedPlatformsFromDeployDeviceProxy(Profile);
-
-	OnPropertyChanged();
-}
-
-void FBasicProfileMarshaller::SetBuildConfiguration(EBuildConfiguration BuildConfiguration)
-{
-	Profile->SetBuildConfiguration(BuildConfiguration);
-	OnPropertyChanged();
-}
-
-EBuildConfiguration FBasicProfileMarshaller::GetBuildConfiguration() const
-{
-	return Profile->GetBuildConfiguration();
-}
-
-void FBasicProfileMarshaller::SetContentScheme(CustomLaunchUI::EContentScheme InContentScheme)
+namespace CustomLaunchUI
 {
-	Model->SetProfileContentScheme(InContentScheme, Profile);
-	OnPropertyChanged();
-
-	// refresh the cached content scheme in case the option that was selected is not available
-	ContentScheme = Model->DetermineProfileContentScheme(Profile);
-}
-
-bool FBasicProfileMarshaller::IsContentSchemeAvailable(CustomLaunchUI::EContentScheme InContentScheme, FText& OutReason) const
-{
-	CustomLaunchUI::FProjectSettings ProjectSettings = Model->GetProjectSettings(Profile);
-
-	// loose files can't be selected if the project is using zen store because there's no way to opt-out of Zen Store from the UAT command line.
-	if (InContentScheme == CustomLaunchUI::EContentScheme::LooseFiles && ProjectSettings.bUseZenStore)
+	FBasicProfileMarshaller::FBasicProfileMarshaller( const ILauncherProfileRef& InProfile, const TSharedRef<FModel>& InModel )
+		: FGenericProfileMarshaller( InProfile, InModel->GetDefaultBasicLaunchProfile(), InModel )
 	{
-		OutReason = LOCTEXT("NoLooseFilesReason", "Loose Files cannot be used when Zen Store is enabled in Project Settings");
-		return false;
 	}
 
-	// cannot launch via Zen if we're targeting a remote device and Zen isn't accepting external connections
-	if (InContentScheme == CustomLaunchUI::EContentScheme::ZenStreaming && !ProjectSettings.bAllowRemoteNetworkService && !Model->IsHostPlatform(Profile))
-	{
-		OutReason = LOCTEXT("NoZenReason", "Zen Streaming to a remote device requires AllowRemoteNetworkService");
-		return false;
-	}
 
-	// basic launch is aimed at launching the current content, not an external build
-	if (InContentScheme == CustomLaunchUI::EContentScheme::ZenPakStreaming)
+	void FBasicProfileMarshaller::Construct()
 	{
-		// don't set a reason - just hide the item for Basic Launch
-		return false;
-	}
+		FGenericProfileMarshaller::Construct();
 
-	return true;
-}
-
-FText FBasicProfileMarshaller::GetCommandLine() const
-{
-	// get unified command line from these two fields. the first is presented in "Build" for old Project Launcher, and the latter is presented in "Launch" for old Project Launcher)
-	// when we save back to the profile, this will be stored just in the "Build" one for clarity (because multiple roles are not supported in old or new project launcher)
-	FString CommandLine = Profile->GetAdditionalCommandLineParameters().TrimStartAndEnd() + TEXT(" ") + Profile->GetDefaultLaunchRole()->GetUATCommandLine().TrimStartAndEnd();
-	CommandLine.TrimStartAndEndInline();
-
-	return FText::FromString(CommandLine);
-}
-
-void FBasicProfileMarshaller::SetCommandLine( const FText& NewText, ETextCommit::Type InTextCommit )
-{
-	FString CommandLine = NewText.ToString();
-	Profile->SetAdditionalCommandLineParameters(CommandLine);
-	Profile->GetDefaultLaunchRole()->SetCommandLine(TEXT(""));
-
-	OnPropertyChanged();
-}
-
-void FBasicProfileMarshaller::SetDeployDeviceIDs(TArray<FString> DeployDeviceIDs)
-{
-	if (Profile->GetDeployedDeviceGroup() == nullptr)
-	{
-		Profile->SetDeployedDeviceGroup(Model->GetProfileManager()->AddNewDeviceGroup());
+		FLaunchProfileTreeNode& GeneralSettingsHeader = TreeData->AddHeading( LOCTEXT("GeneralSettingsHeading", "General Settings") );
+		AddProjectProperty(GeneralSettingsHeader);
+		AddTargetProperty(GeneralSettingsHeader);
+		AddConfigurationProperty(GeneralSettingsHeader);
+		AddContentSchemeProperty(GeneralSettingsHeader);
+		AddCommandLineProperty(GeneralSettingsHeader);
+		AddTargetDeviceProperty(GeneralSettingsHeader);
 	}
 
-	Profile->GetDeployedDeviceGroup()->RemoveAllDevices();
-	for (const FString& DeployDeviceID : DeployDeviceIDs)
-	{
-		Profile->GetDeployedDeviceGroup()->AddDevice(DeployDeviceID);
-	}
-	Model->UpdatedCookedPlatformsFromDeployDeviceProxy(Profile);
 
-	OnPropertyChanged();
-}
 
-TArray<FString> FBasicProfileMarshaller::GetDeployDeviceIDs() const
-{
-	return Profile->GetDeployedDeviceGroup()->GetDeviceIDs();
-}
-
-
-void FBasicProfileMarshaller::OnPropertyChanged()
-{
-	// do not save basic profiles - they're transient
-}
-
-
-
-TSharedPtr<ILaunchProfileMarshaller> FBasicProfileMarshallerFactory::TryCreateMarshaller( const ILauncherProfileRef& InProfile, const TSharedRef<CustomLaunchUI::FModel>& InModel )
-{
-	if (InModel->IsBasicLaunchProfile(InProfile))
+	TSharedPtr<ILaunchProfileMarshaller> FBasicProfileMarshallerFactory::TryCreateMarshaller( const ILauncherProfileRef& InProfile, const TSharedRef<FModel>& InModel )
 	{
 		return MakeShared<FBasicProfileMarshaller>(InProfile, InModel);
 	}
 
-	return nullptr;
-}
-
 
+}
 
 #undef LOCTEXT_NAMESPACE
diff --git a/Engine/Plugins/Developer/CustomLaunchUI/Source/Private/Marshalling/BasicProfileMarshaller.h b/Engine/Plugins/Developer/CustomLaunchUI/Source/Private/Marshalling/BasicProfileMarshaller.h
deleted file mode 100644
diff --git a/Engine/Plugins/Developer/CustomLaunchUI/Source/Private/Marshalling/CustomProfileMarshaller.cpp b/Engine/Plugins/Developer/CustomLaunchUI/Source/Private/Marshalling/CustomProfileMarshaller.cpp
--- a/Engine/Plugins/Developer/CustomLaunchUI/Source/Private/Marshalling/CustomProfileMarshaller.cpp
+++ b/Engine/Plugins/Developer/CustomLaunchUI/Source/Private/Marshalling/CustomProfileMarshaller.cpp
@@ -1,1257 +1,66 @@
 // Copyright Epic Games, Inc. All Rights Reserved.
 
 #include "Marshalling/CustomProfileMarshaller.h"
-#include "Model/CustomLaunchUIModel.h"
-#include "Widgets/Shared/SCustomLaunchProjectCombo.h"
-#include "Widgets/Shared/SCustomLaunchBuildTargetCombo.h"
-#include "Widgets/Shared/SCustomLaunchPlatformCombo.h"
-#include "Widgets/Shared/SCustomLaunchContentSchemeCombo.h"
-#include "Widgets/Shared/SCustomLaunchDeviceListView.h"
-#include "Widgets/Shared/SCustomLaunchMapListView.h"
-#include "Widgets/Shared/SCustomLaunchCombo.h"
-#include "Widgets/Input/SSegmentedControl.h"
-#include "Widgets/Text/STextBlock.h"
-#include "Widgets/Layout/SBox.h"
-#include "Widgets/SBoxPanel.h"
-#include "SResizeBox.h"
-#include "SSearchableComboBox.h"
-#include "SlateOptMacros.h"
-#include "PlatformInfo.h"
-#include "IDesktopPlatform.h"
-#include "Interfaces/ITargetPlatformManagerModule.h"
-#include "Interfaces/ITargetPlatform.h"
-#include "GameProjectHelper.h"
-
 
 
 #define LOCTEXT_NAMESPACE "CustomProfileMarshaller"
 
-class FCustomProfileMarshaller : public ILaunchProfileMarshaller, public TSharedFromThis<FCustomProfileMarshaller>
-{
-public:
-	FCustomProfileMarshaller( const ILauncherProfileRef& Profile, const TSharedRef<CustomLaunchUI::FModel>& InModel );
-	virtual ~FCustomProfileMarshaller() = default;
-
-	virtual void Construct() override;
-
-	virtual FString GetName() const override
-	{
-		return TEXT("CustomProfile");
-	}
-
-	virtual FLaunchProfileTreeDataRef GetProfileTree() override
-	{
-		return TreeData;
-	}
-
-protected:
-	TArray<FString> CachedMapsToCook;
-	enum class EMapOption : uint8
-	{
-		Startup,
-		Selected,
-	};
-	EMapOption MapOption = EMapOption::Startup;
-
-	TArray<FString> CachedDeployDeviceIDs;
-	enum class EDeployDeviceOption : uint8
-	{
-		Default,
-		Selected,
-	};
-	EDeployDeviceOption DeployDeviceOption = EDeployDeviceOption::Default;
-
-
-	FString GetProjectPath() const;
-	void SetProjectName(FString ProjectPath);
-	bool HasProject() const;
-
-	FString GetBuildTarget() const;
-	void SetBuildTarget(FString BuildTarget);
-	TArray<EBuildTargetType> GetSupportedBuildTargetTypes() const;
-
-	void SetBuildConfiguration(EBuildConfiguration BuildConfiguration);
-	EBuildConfiguration GetBuildConfiguration() const;
-
-	void SetContentScheme(CustomLaunchUI::EContentScheme ContentScheme);
-	bool IsContentSchemeAvailable(CustomLaunchUI::EContentScheme, FText& OutReason) const;
-
-	FString GetCommandLine() const;
-	void SetCommandLine( const FString& NewCommandLine );
-
-	void SetSelectedPlatforms( TArray<FString> SelectedPlatforms );
-	TArray<FString> GetSelectedPlatforms() const;
-
-	void SetCook( bool bCook );
-	bool GetCook( ILauncherProfilePtr InProfile = nullptr ) const;
-
-	void SetMapsToCook(TArray<FString> MapsToCook);
-	TArray<FString> GetMapsToCook() const;
-	EMapOption GetMapOption() const;
-	void SetMapOption( EMapOption MapOption );
-	float GetMapListHeight() const;
-	void SetMapListHeight( float NewHeight );
-	TSharedRef<SWidget> CreateMapListWidget();
-	float MapListHeight = 300.0f;
-
-
-	void SetDeployDeviceIDs(TArray<FString> DeployDeviceIDs);
-	TArray<FString> GetDeployDeviceIDs() const;
-	EDeployDeviceOption GetDeployDeviceOption() const;
-	void SetDeployDeviceOption( EDeployDeviceOption DeployDeviceOption );
-	float GetDeployDeviceListHeight() const;
-	void SetDeployDeviceListHeight( float NewHeight );
-	void OnDeviceRemoved(FString DeviceID);
-	TSharedRef<SWidget> CreateDeployDeviceWidget();
-	float DeployDeviceListHeight = 150.0f;
-
-	void SetBuild( bool bBuild );
-	bool GetBuild( ILauncherProfilePtr InProfile = nullptr ) const;
-
-	void SetForceBuild( bool bForceBuild );
-	bool GetForceBuild( ILauncherProfilePtr InProfile = nullptr ) const;
-
-	void SetArchitecture( FString Architecture );
-	FString GetArchitecture() const;
-	FText GetArchitectureDisplayName( FString Architecture );
-
-	void SetDeployToDevice( bool bDeployToDevice );
-	bool GetDeployToDevice( ILauncherProfilePtr InProfile = nullptr ) const;
-
-	void SetIsRunning( bool bRun );
-	bool GetIsRunning( ILauncherProfilePtr InProfile = nullptr ) const;
-
-	void OnInitialMapChanged( TSharedPtr<FString> InitialMap, ESelectInfo::Type );
-	TSharedPtr<FString> GetInitialMap() const;
-
-
-
-protected:
-
-	virtual void OnPropertyChanged() override;
-	void RefreshContentScheme();
-	void CacheStartupMapList();
-	void CacheArchitectures();
-	TSharedRef<SWidget> OnGenerateComboWidget( TSharedPtr<FString> InComboString );
-
-
-	FLaunchProfileTreeDataRef TreeData;
-	const ILauncherProfileRef Profile;
-	const ILauncherProfileRef DefaultProfile;
-	const TSharedRef<CustomLaunchUI::FModel> Model;
-
-	CustomLaunchUI::EContentScheme ContentScheme;
-	bool bShouldCook;
-	TArray<TSharedPtr<FString>> CachedStartupMaps;
-	TSharedPtr<SSearchableComboBox> InitalMapCombo;
-	TSharedPtr<SCustomLaunchMapListView> MapListView;
-	TSharedPtr<SCustomLaunchDeviceListView> DeployDeviceListView;
-	TArray<FString> CachedArchitectures;
-	EBuildTargetType CachedBuildTargetType = EBuildTargetType::Game;
-};
-
-
-FCustomProfileMarshaller::FCustomProfileMarshaller( const ILauncherProfileRef& InProfile, const TSharedRef<CustomLaunchUI::FModel>& InModel )
-	: TreeData( MakeShared<FLaunchProfileTreeData>(InProfile, InModel, this) )
-	, Profile(InProfile)
-	, DefaultProfile(InModel->GetDefaultCustomLaunchProfile())
-	, Model(InModel)
-{
-}
-
-BEGIN_SLATE_FUNCTION_BUILD_OPTIMIZATION
-void FCustomProfileMarshaller::Construct()
-{
-	auto ForPak = [this]()
-	{
-		return (ContentScheme == CustomLaunchUI::EContentScheme::PakFiles || ContentScheme == CustomLaunchUI::EContentScheme::DevelopmentPackage);
-	};
-
-	auto ForZenWS = [this]()
-	{
-		return (ContentScheme == CustomLaunchUI::EContentScheme::ZenPakStreaming);
-	};
-
-	auto ForCooked = [this]()
-	{
-		return (ContentScheme != CustomLaunchUI::EContentScheme::ZenPakStreaming && ContentScheme != CustomLaunchUI::EContentScheme::CookOnTheFly);
-	};
-
-	auto ForEnabledCooked = [this]()
-	{
-		return (bShouldCook && ContentScheme != CustomLaunchUI::EContentScheme::ZenPakStreaming && ContentScheme != CustomLaunchUI::EContentScheme::CookOnTheFly);
-	};
-
-	auto ForCode = [this]()
-	{
-		return GetBuild();
-	};
-
-	auto ForDeployment = [this]()
-	{
-		return GetDeployToDevice();
-	};
-
-	auto ForRun = [this]()
-	{
-		return GetIsRunning();
-	};
-
-	auto EmptyString = []()
-	{
-		return FString();
-	};
-
-
-
-	TreeData->AddHeading( LOCTEXT("GeneralSettingsHeading", "General Settings") )
-
-		.AddWidget( LOCTEXT("ProjectLabel", "Project"), 
-			SNew(SCustomLaunchProjectCombo)
-			.OnSelectionChanged(this, &FCustomProfileMarshaller::SetProjectName)
-			.SelectedProject(this, &FCustomProfileMarshaller::GetProjectPath)
-			.HasProject(this, &FCustomProfileMarshaller::HasProject)
-			.ShowAnyProjectOption(true)
-			.CurrentProjectOption(SCustomLaunchProjectCombo::ECurrentProjectOption::ActualProject)
-			.Font(FCoreStyle::Get().GetFontStyle("SmallFont"))
-		)
-
-		.AddWidget( LOCTEXT("TargetLabel", "Target"), 
-			{
-				.IsDefault = [this]()		{ return !Profile->HasBuildTargetSpecified() || Profile->GetBuildTarget().IsEmpty(); },
-				.SetToDefault = [this]()	{ SetBuildTarget(FString()); },
-				.IsEnabled = [this]()		{ return Profile->HasProjectSpecified(); },
-			},
-			SNew(SCustomLaunchBuildTargetCombo)
-			.OnSelectionChanged(this, &FCustomProfileMarshaller::SetBuildTarget)
-			.SelectedBuildTarget(this, &FCustomProfileMarshaller::GetBuildTarget)
-			.SupportedTargetTypes(this, &FCustomProfileMarshaller::GetSupportedBuildTargetTypes)
-			.SelectedProject(this, &FCustomProfileMarshaller::GetProjectPath)
-			.Font(FCoreStyle::Get().GetFontStyle("SmallFont"))
-		)
-
-		.AddWidget( LOCTEXT("PlatformLabel", "Platform"), 
-			SNew(SCustomLaunchPlatformCombo)
-			.SelectedPlatforms(this, &FCustomProfileMarshaller::GetSelectedPlatforms)
-			.OnSelectionChanged(this, &FCustomProfileMarshaller::SetSelectedPlatforms)
-			.BasicPlatformsOnly(CustomLaunchUI::bUseFriendlyBuildTargetSelection)
-			.Font(FCoreStyle::Get().GetFontStyle("SmallFont"))
-		)
-
-		.AddWidget( LOCTEXT("ConfigurationLabel", "Configuration"), 
-			SNew(SCustomLaunchLexToStringCombo<EBuildConfiguration>)
-			.OnSelectionChanged(this, &FCustomProfileMarshaller::SetBuildConfiguration)
-			.SelectedItem(this, &FCustomProfileMarshaller::GetBuildConfiguration)
-			.Items( TArray<EBuildConfiguration>({ EBuildConfiguration::Debug, EBuildConfiguration::Development, EBuildConfiguration::Test, EBuildConfiguration::Shipping }) )
-		)
-
-		.AddWidget( LOCTEXT("ContentSchemeLabel", "Content Scheme"), 
-			SNew(SCustomLaunchContentSchemeCombo)
-			.OnSelectionChanged(this, &FCustomProfileMarshaller::SetContentScheme)
-			.SelectedContentScheme_Lambda( [this]() { return ContentScheme; } )
-			.IsContentSchemeAvailable(this, &FCustomProfileMarshaller::IsContentSchemeAvailable)
-			.Font(FCoreStyle::Get().GetFontStyle("SmallFont"))
-		)
-	;
-
-
-	TreeData->AddHeading( LOCTEXT("ContentSchemeHeading", "Content Scheme") )
-
-		.AddBoolean( LOCTEXT("CompressPakFilesLabel", "Compress Pak Files"),
-			{
-				.GetValue = [this]()			{ return Profile->IsCompressed(); },
-				.SetValue = [this](bool bValue)	{ Profile->SetCompressed(bValue); },
-				.GetDefaultValue = [this]()		{ return DefaultProfile->IsCompressed(); },
-				.IsVisible = ForPak,
-			}
-		)
-
-		.AddBoolean( LOCTEXT("UseIoStoreLabel", "Use Io Store"),
-			{
-				.GetValue = [this]()			{ return Profile->IsUsingIoStore(); },
-				.SetValue = [this](bool bValue)	{ Profile->SetUseIoStore(bValue); },
-				.GetDefaultValue = [this]()		{ return DefaultProfile->IsUsingIoStore(); },
-				.IsVisible = ForPak,
-			}
-		)
-
-		.AddBoolean( LOCTEXT("GenerateChunksLabel", "Generate Chunks"),
-			{
-				.GetValue = [this]()			{ return Profile->IsGeneratingChunks(); },
-				.SetValue = [this](bool bValue)	{ Profile->SetGenerateChunks(bValue); },
-				.GetDefaultValue = [this]()		{ return DefaultProfile->IsGeneratingChunks(); },
-				.IsVisible = ForPak,
-			}
-		)
-
-		.AddBoolean( LOCTEXT("ImportZenSnapshotLabel", "Import Best Match Zen Snapshot"),
-			{
-				.GetValue = [this]()			{ return Profile->IsImportingZenSnapshot(); },
-				.SetValue = [this](bool bValue)	{ Profile->SetImportingZenSnapshot(bValue); },
-				.GetDefaultValue = [this]()		{ return DefaultProfile->IsImportingZenSnapshot(); },
-				.IsVisible = ForCooked,
-			}
-		)
-
-		.AddDirectoryString( LOCTEXT("ZenPakStreamingPathLabel", "Zen Pak Streaming Path"),
-			{
-				.GetValue = [this]()				{ return Profile->GetZenPakStreamingPath(); },
-				.SetValue = [this](FString Value)	{ Profile->SetZenPakStreamingPath(Value); },
-				.GetDefaultValue = EmptyString,
-				.IsVisible = ForZenWS,
-			}
-		)
-
-		.AddBoolean( LOCTEXT("IncrementalCookLabel", "Incremental Cook"),
-			{
-				.GetValue = [this]()			{ return Profile->IsCookingIncrementally() || Profile->IsImportingZenSnapshot(); }, // show it as ticked when snapshot import is enabled for now to indicate that it will be doing something incremental. eventually this is likely to become a separate field
-				.SetValue = [this](bool bValue)	{ Profile->SetIncrementalCooking(bValue); },
-				.GetDefaultValue = [this]()		{ return DefaultProfile->IsCookingIncrementally(); },
-				.IsVisible = ForCooked,
-				.IsEnabled = [this]()			{ return !Profile->IsImportingZenSnapshot(); }, // zen snapshot import & incremental cooking can't be used together. '-iterativecook' is the 'old' system, and zen import will enable the new '-cookincremental' & the two flags don't play nicely together
-			}
-		)
-	;
-
-
-	TreeData->AddHeading( LOCTEXT("CookingHeading", "Maps And Cooking") )
-
-		.AddBoolean( LOCTEXT("CookLabel", "Cook Content"),
-			{
-				.GetValue = [this]()			{ return GetCook(); },
-				.SetValue = [this](bool bValue)	{ SetCook(bValue); },
-				.GetDefaultValue = [this]()		{ return GetCook(DefaultProfile); },
-				.IsVisible = ForCooked
-			}
-		)
-
-		.AddWidget( LOCTEXT("MapsToCookLabel", "Maps To Cook"), 
-			{
-				.IsVisible = ForEnabledCooked,
-			},
-			CreateMapListWidget()
-		)
-
-		.AddString( LOCTEXT("AdditionalCookerOptionsLabel", "Additional Cooker Options"), 
-			{
-				.GetValue = [this]()				{ return Profile->GetCookOptions(); },
-				.SetValue = [this](FString Value)	{ Profile->SetCookOptions( Value ); },
-				.GetDefaultValue = EmptyString,
-				.IsVisible = ForEnabledCooked,
-			}
-		)
-	;
-
-
-	TreeData->AddHeading( LOCTEXT("BuildHeading", "Build") )
-
-		.AddBoolean( LOCTEXT("BuildLabel", "Build the game"),
-			{
-				.GetValue = [this]()			{ return GetBuild(); },
-				.SetValue = [this](bool bValue)	{ SetBuild(bValue); },
-				.GetDefaultValue = [this]()		{ return GetBuild(DefaultProfile); },
-			}
-		)
-
-		.AddBoolean( LOCTEXT("ForceBuildLabel", "Build even if a pre-built target exists"),
-			{
-				.GetValue = [this]()			{ return GetForceBuild(); },
-				.SetValue = [this](bool bValue)	{ SetForceBuild(bValue); },
-				.GetDefaultValue = [this]()		{ return GetForceBuild(DefaultProfile); },
-				.IsEnabled = ForCode,
-			}
-		)
-
-		// hiding this option for now - no good reason to buid UAT it for the majority of normal use cases
-		//.AddBoolean( LOCTEXT("BuildUATLabel", "Build UAT"),
-		//	{
-		//		.GetValue = [this]()			{ return Profile->IsBuildingUAT(); },
-		//		.SetValue = [this](bool bValue)	{ Profile->SetBuildUAT(bValue); },
-		//		.GetDefaultValue = [this]()		{ return DefaultProfile->IsBuildingUAT(); },
-		//		.IsEnabled = ForCode,
-		//	}
-		//)
-
-		.AddWidget( LOCTEXT("ArchitectureLabel", "Architecture"),
-			{
-				.IsDefault = [this]()		{ return GetArchitecture().IsEmpty(); },
-				.SetToDefault = [this]()	{ SetArchitecture(FString()); },
-				.IsVisible = [this]()       { return CachedArchitectures.Num() > 0; }
-			},
-			SNew(SCustomLaunchStringCombo)
-			.OnSelectionChanged( this, &FCustomProfileMarshaller::SetArchitecture )
-			.SelectedItem( this, &FCustomProfileMarshaller::GetArchitecture)
-			.GetDisplayName( this, &FCustomProfileMarshaller::GetArchitectureDisplayName )
-			.Items_Lambda( [this]() { return CachedArchitectures;} )
-		)
-	;
-
-
-
-	TreeData->AddHeading( LOCTEXT("DirectoryHeading", "Directory") )
-
-		// unlikely to want to change this in most use cases
-		//.AddDirectoryString( LOCTEXT("CustomStagingPathLabel", "Custom Stage Directory"),
-		//	{
-		//		.GetValue = [this]()				{ return Profile->GetPackageDirectory(); },
-		//		.SetValue = [this](FString Value)	{ Profile->SetPackageDirectory(Value); },
-		//		.GetDefaultValue = EmptyString,
-		//		.IsVisible = ForCooked,
-		//	}
-		//)
-
-		.AddBoolean( LOCTEXT("ArchiveBuildLabel", "Archive Build"),
-			{
-				.GetValue = [this]()			{ return Profile->IsArchiving(); },
-				.SetValue = [this](bool bValue)	{ Profile->SetArchive(bValue); },
-				.GetDefaultValue = [this]()		{ return DefaultProfile->IsArchiving(); },
-				.IsVisible = ForCooked,
-			}
-		)
-
-		.AddDirectoryString( LOCTEXT("ArchivePathLabel", "Archive Directory"),
-			{
-				.GetValue = [this]()				{ return Profile->GetArchiveDirectory(); },
-				.SetValue = [this](FString Value)	{ Profile->SetArchiveDirectory(Value); },
-				.GetDefaultValue = EmptyString,
-				.IsVisible = ForCooked,				
-				.IsEnabled = [this]()				{ return Profile->IsArchiving(); }
-			}
-		)
-	;
-
-
-
-	TreeData->AddHeading( LOCTEXT("DeployAndRunHeading", "Deploy And Run") )
-
-		.AddBoolean( LOCTEXT("DeployLabel", "Deploy To Device"),
-			{
-				.GetValue = [this]()			{ return GetDeployToDevice(); },
-				.SetValue = [this](bool bValue)	{ SetDeployToDevice(bValue); },
-				.GetDefaultValue = [this]()		{ return GetDeployToDevice(DefaultProfile); },
-				.IsVisible = [this]()			{ return (ContentScheme != CustomLaunchUI::EContentScheme::CookOnTheFly); }
-			}
-		)
-
-		// not relevant to most scenarios now
-		//.AddBoolean( LOCTEXT("DeployModifiedLabel", "Only Deploy Modified Content"),
-		//	{
-		//		.GetValue = [this]()			{ return Profile->IsDeployingIncrementally(); },
-		//		.SetValue = [this](bool bValue)	{ Profile->SetIncrementalDeploying(bValue); },
-		//		.GetDefaultValue = [this]()		{ return DefaultProfile->IsDeployingIncrementally(); },
-		//		.IsEnabled = ForDeployment,
-		//	}
-		//)
-
-		.AddWidget( LOCTEXT("TargetDeviceLabel", "Target Device"), 
-			{
-				.IsEnabled = ForDeployment,
-			},
-			CreateDeployDeviceWidget()
-		)
-
-		.AddBoolean( LOCTEXT("RunLabel", "Run"),
-			{
-				.GetValue = [this]()			{ return GetIsRunning(); },
-				.SetValue = [this](bool bValue)	{ SetIsRunning(bValue); },
-				.GetDefaultValue = [this]()		{ return GetIsRunning(DefaultProfile); },
-			}
-		)
-
-		// todo: custom map picker 
-		.AddWidget( LOCTEXT("InitialMapLabel", "Initial Map"), 
-			{
-				.IsDefault = [this]()		{ return Profile->GetDefaultLaunchRole()->GetInitialMap().IsEmpty(); },
-				.SetToDefault = [this]()	{ Profile->GetDefaultLaunchRole()->SetInitialMap(FString()); InitalMapCombo->SetSelectedItem(GetInitialMap()); },
-				.IsEnabled = ForRun,
-			},
-			SAssignNew(InitalMapCombo, SSearchableComboBox)
-			.OptionsSource(&CachedStartupMaps)
-			.OnSelectionChanged( this, &FCustomProfileMarshaller::OnInitialMapChanged )
-			.OnGenerateWidget( this, &FCustomProfileMarshaller::OnGenerateComboWidget )
-			.Content()
-			[
-				SNew(STextBlock)
-				.Text_Lambda( [this]() { return FText::FromString( *GetInitialMap() ); } )
-				.Font(FCoreStyle::Get().GetFontStyle("SmallFont"))
-			]
-		)
-
-		.AddString( LOCTEXT("CommandLineLabel", "Command Line"), 
-			{
-				.GetValue = [this]()				{ return GetCommandLine(); },
-				.SetValue = [this](FString Value)	{ SetCommandLine(Value); },
-				.GetDefaultValue = EmptyString,
-				.IsEnabled = ForRun
-			}
-		)
-	;
-
-
-
-	MapOption = Profile->GetCookedMaps().Num() > 0 ? EMapOption::Selected : EMapOption::Startup;
-	DeployDeviceOption = Profile->GetDeployedDeviceGroup()->GetDeviceIDs().Num() > 0 ? EDeployDeviceOption::Selected : EDeployDeviceOption::Default;
-	bShouldCook = GetCook();
-	ContentScheme = Model->DetermineProfileContentScheme(Profile);
-	CachedBuildTargetType = Model->GetBuildTargetInfo(Profile).Type;
-	
-	CacheStartupMapList();
-	CacheArchitectures();
-}
-END_SLATE_FUNCTION_BUILD_OPTIMIZATION
-
-
-
-void FCustomProfileMarshaller::OnPropertyChanged()
-{
-	Model->GetProfileManager()->SaveJSONProfile(Profile);
-}
-
-
-void FCustomProfileMarshaller::CacheStartupMapList()
-{
-	CachedStartupMaps.Reset();
-	CachedStartupMaps.Add(MakeShared<FString>());
-	for ( const FString& Map : FGameProjectHelper::GetAvailableMaps(Profile->GetProjectBasePath(), Profile->SupportsEngineMaps(), true))
-	{
-		CachedStartupMaps.Add(MakeShared<FString>(Map));
-	}
-
-	InitalMapCombo->SetSelectedItem( GetInitialMap(), ESelectInfo::Direct );
-	InitalMapCombo->RefreshOptions();
-}
-
-void FCustomProfileMarshaller::CacheArchitectures()
-{
-	CachedArchitectures.Reset();
-
-	const PlatformInfo::FTargetPlatformInfo* PlatformInfo = Model->GetPlatformInfo(Profile);
-	if (PlatformInfo != nullptr)
-	{
-		const ITargetPlatform* TargetPlatform = GetTargetPlatformManager()->FindTargetPlatform(PlatformInfo->IniPlatformName);
-		if (TargetPlatform)
-		{
-			TargetPlatform->GetPossibleArchitectures(CachedArchitectures);
-			if (CachedArchitectures.Num() > 0)
-			{
-				CachedArchitectures.Insert(FString(), 0); // empty string for "project default" option
-			}
-		}
-	}
-}
-
-
-
-TSharedRef<SWidget> FCustomProfileMarshaller::OnGenerateComboWidget( TSharedPtr<FString> InComboString )
-{
-	return SNew(STextBlock)
-		.Text(InComboString.IsValid() ? FText::FromString(*InComboString) : FText::GetEmpty() )
-		.Font(FCoreStyle::Get().GetFontStyle("SmallFont"));
-
-}
-
-
-
-void FCustomProfileMarshaller::SetSelectedPlatforms( TArray<FString> SelectedPlatforms )
-{
-	Profile->ClearCookedPlatforms();
-
-	if (CustomLaunchUI::bUseFriendlyBuildTargetSelection)
-	{
-		FTargetInfo BuildTargetInfo = CustomLaunchUI::FModel::GetBuildTargetInfo(GetBuildTarget(), GetProjectPath());
-
-		for (const FString& Platform : SelectedPlatforms)
-		{
-			Profile->AddCookedPlatform( CustomLaunchUI::FModel::GetBuildTargetPlatformName(Platform, BuildTargetInfo));
-		}
-	}
-	else
-	{
-		for (const FString& Platform : SelectedPlatforms)
-		{
-			Profile->AddCookedPlatform(Platform);
-		}
-	}
-	OnPropertyChanged();
-
-	DeployDeviceListView->OnSelectedPlatformChanged();
-
-	CacheArchitectures();
-}
-
-TArray<FString> FCustomProfileMarshaller::GetSelectedPlatforms() const
-{
-	TArray<FString> Platforms;
-
-	if (CustomLaunchUI::bUseFriendlyBuildTargetSelection)
-	{
-		for (const FString& Platform : Profile->GetCookedPlatforms())
-		{
-			Platforms.Add(CustomLaunchUI::FModel::GetVanillaPlatformName(Platform));
-		}
-	}
-	else
-	{
-		Platforms = Profile->GetCookedPlatforms();
-	}
-
-	return MoveTemp(Platforms);
-}
-
-
-FString FCustomProfileMarshaller::GetProjectPath() const
-{
-	if (Profile->HasProjectSpecified())
-	{
-		return Profile->GetProjectPath();
-	}
-
-	return FString();
-}
-
-void FCustomProfileMarshaller::SetProjectName(FString ProjectPath)
-{
-	Profile->SetProjectSpecified(!ProjectPath.IsEmpty());
-	Profile->SetProjectPath(ProjectPath);
-
-	if (CustomLaunchUI::bUseFriendlyBuildTargetSelection)
-	{
-		Model->UpdateCookedPlatformsFromBuildTarget(Profile);
-	}
-
-	OnPropertyChanged();
-
-	if (MapListView.IsValid())
-	{
-		MapListView->RefreshMapList();
-	}
-
-	CacheStartupMapList();
-}
-
-bool FCustomProfileMarshaller::HasProject() const
-{
-	return Profile->HasProjectSpecified();
-}
-
-
-FString FCustomProfileMarshaller::GetBuildTarget() const
-{
-	return Profile->GetBuildTarget();
-}
-
-void FCustomProfileMarshaller::SetBuildTarget(FString BuildTarget)
-{
-	Profile->SetBuildTargetSpecified(!BuildTarget.IsEmpty());
-	Profile->SetBuildTarget(BuildTarget);
-
-	if (CustomLaunchUI::bUseFriendlyBuildTargetSelection)
-	{
-		Model->UpdateCookedPlatformsFromBuildTarget(Profile);
-	}
-
-	CachedBuildTargetType = Model->GetBuildTargetInfo(Profile).Type;
-	OnPropertyChanged();
-}
-
-TArray<EBuildTargetType> FCustomProfileMarshaller::GetSupportedBuildTargetTypes() const
-{
-	TArray<EBuildTargetType> Result;
-
-	if (Profile->GetCookedPlatforms().Num() > 0)
-	{
-		const PlatformInfo::FTargetPlatformInfo* PlatformInfo = PlatformInfo::FindPlatformInfo(FName(Profile->GetCookedPlatforms()[0]));
-		if (PlatformInfo)
-		{
-			if (CustomLaunchUI::bUseFriendlyBuildTargetSelection)
-			{
-				Result.Add(PlatformInfo->VanillaInfo->PlatformType);
-
-				for (const PlatformInfo::FTargetPlatformInfo* PlatformFlavorInfo : PlatformInfo->VanillaInfo->Flavors)
-				{
-					Result.AddUnique(PlatformFlavorInfo->PlatformType);
-				}
-			}
-			else
-			{
-				Result.Add(PlatformInfo->PlatformType);
-			}
-		}
-	}
-
-	return MoveTemp(Result);
-}
-
-void FCustomProfileMarshaller::SetBuildConfiguration(EBuildConfiguration BuildConfiguration)
-{
-	Profile->SetBuildConfiguration(BuildConfiguration);
-	OnPropertyChanged();
-}
-
-EBuildConfiguration FCustomProfileMarshaller::GetBuildConfiguration() const
-{
-	return Profile->GetBuildConfiguration();
-}
-
-
-
-void FCustomProfileMarshaller::RefreshContentScheme()
-{
-	CustomLaunchUI::EContentScheme CurrentContentScheme = Model->DetermineProfileContentScheme(Profile);
-	SetContentScheme(CurrentContentScheme);
-}
-
-void FCustomProfileMarshaller::SetContentScheme(CustomLaunchUI::EContentScheme InContentScheme)
-{
-	ELauncherProfileDeploymentModes::Type DeploymentMode = GetDeployToDevice() ? ELauncherProfileDeploymentModes::CopyToDevice : ELauncherProfileDeploymentModes::DoNotDeploy;
-	Model->SetProfileContentScheme(InContentScheme, Profile, bShouldCook, DeploymentMode );
-	OnPropertyChanged();
-
-	// refresh the cached content scheme in case the option that was selected is not available
-	ContentScheme = Model->DetermineProfileContentScheme(Profile);
-}
-
-bool FCustomProfileMarshaller::IsContentSchemeAvailable(CustomLaunchUI::EContentScheme InContentScheme, FText& OutReason) const
-{
-	CustomLaunchUI::FProjectSettings ProjectSettings = Model->GetProjectSettings(Profile);
-
-	// loose files can't be selected if the project is using zen store because there's no way to opt-out of Zen Store from the UAT command line.
-	if (InContentScheme == CustomLaunchUI::EContentScheme::LooseFiles && ProjectSettings.bUseZenStore)
-	{
-		OutReason = LOCTEXT("NoLooseFilesReason", "Loose Files cannot be used when Zen Store is enabled in Project Settings");
-		return false;
-	}
-
-	// don't show zen pak streaming option if it isn't going to be set up automatically by UAT
-	// @todo: could potentially look at shelling out to Zen & querying if we have a dynamic workspace? no support for async config queries in this tool yet though
-	if (InContentScheme == CustomLaunchUI::EContentScheme::ZenPakStreaming && !ProjectSettings.bHasAutomaticZenPakStreamingWorkspaceCreation)
-	{
-		OutReason = LOCTEXT("NoZenPakReason", "Automatic Zen Pak streaming workspace creation has not been enabled in Project Settings");
-		return false;
-	}
-	
-	// cannot launch via Zen if we're targeting a remote device and Zen isn't accepting external connections
-	if ((InContentScheme == CustomLaunchUI::EContentScheme::ZenStreaming || InContentScheme == CustomLaunchUI::EContentScheme::ZenPakStreaming) 
-		&& !ProjectSettings.bAllowRemoteNetworkService && !Model->IsHostPlatform(Profile))
-	{
-		OutReason = LOCTEXT("NoZenReason", "Zen Streaming to a remote device requires AllowRemoteNetworkService");
-		return false;
-	}
-
-	return true;
-}
-
-
-
-FString FCustomProfileMarshaller::GetCommandLine() const
-{
-	// get unified command line from these two fields. the first is presented in "Build" for old Project Launcher, and the latter is presented in "Launch" for old Project Launcher)
-	// when we save back to the profile, this will be stored just in the "Build" one for clarity (because multiple roles are not supported in old or new project launcher)
-	FString CommandLine = Profile->GetAdditionalCommandLineParameters().TrimStartAndEnd() + TEXT(" ") + Profile->GetDefaultLaunchRole()->GetUATCommandLine().TrimStartAndEnd();
-	CommandLine.TrimStartAndEndInline();
-
-	return CommandLine;
-}
-
-void FCustomProfileMarshaller::SetCommandLine( const FString& NewCommandLine )
-{
-	Profile->SetAdditionalCommandLineParameters(NewCommandLine);
-	Profile->GetDefaultLaunchRole()->SetCommandLine(TEXT(""));
-	OnPropertyChanged();
-}
-
-
-
-
-BEGIN_SLATE_FUNCTION_BUILD_OPTIMIZATION
-TSharedRef<SWidget> FCustomProfileMarshaller::CreateMapListWidget()
-{
-	MapListView = 
-		SNew(SCustomLaunchMapListView)
-		.OnSelectionChanged(this, &FCustomProfileMarshaller::SetMapsToCook)
-		.SelectedMaps(this, &FCustomProfileMarshaller::GetMapsToCook)
-		.ProjectPath(this, &FCustomProfileMarshaller::GetProjectPath)
-	;
-
-	return SNew(SVerticalBox)
-
-		// map option controls
-		+SVerticalBox::Slot()
-		.AutoHeight()
-		.Padding(0,2)
-		[
-			SNew(SHorizontalBox)
-
-			// map option
-			+SHorizontalBox::Slot()
-			.Padding(0, 0)
-			.AutoWidth()
-			.VAlign(VAlign_Center)
-			[
-				SNew(SSegmentedControl<EMapOption>)
-				.Value( this, &FCustomProfileMarshaller::GetMapOption)
-				.OnValueChanged(this, &FCustomProfileMarshaller::SetMapOption)
-
-				+SSegmentedControl<EMapOption>::Slot(EMapOption::Startup)
-				.HAlign(HAlign_Center)
-				.VAlign(VAlign_Center)
-				[
-					SNew(STextBlock)
-					.Text(LOCTEXT("StartupMapsLabel", "Startup Maps"))
-					.Font(FCoreStyle::Get().GetFontStyle("SmallFont"))
-				]
-
-				+SSegmentedControl<EMapOption>::Slot(EMapOption::Selected)
-				.HAlign(HAlign_Center)
-				.VAlign(VAlign_Center)
-				[
-					SNew(STextBlock)
-					.Text(LOCTEXT("SelectedMapsLabel", "Selected Maps"))
-					.Font(FCoreStyle::Get().GetFontStyle("SmallFont"))
-				]
-			]
-
-			// map selector controls (search etc)
-			+SHorizontalBox::Slot()
-			.Padding(8, 0)
-			.FillWidth(1)
-			.VAlign(VAlign_Center)
-			[
-				SNew(SBox)
-				.Visibility_Lambda([this](){ return GetMapOption() == EMapOption::Selected  ? EVisibility::Visible : EVisibility::Collapsed; } )
-				[
-					MapListView->MakeControlsWidget()
-				]
-			]
-		]
-
-		// map list
-		+SVerticalBox::Slot()
-		.FillHeight(1)
-		[
-			SNew(SVerticalResizeBox)
-			.Visibility_Lambda([this](){ return (GetMapOption() == EMapOption::Selected) ? EVisibility::Visible : EVisibility::Collapsed; } )
-			.HandleHeight(4)
-			.ContentHeight(this, &FCustomProfileMarshaller::GetMapListHeight)
-			.ContentHeightChanged(this, &FCustomProfileMarshaller::SetMapListHeight)
-			.HandleColor( FAppStyle::Get().GetSlateColor("Colors.Secondary").GetSpecifiedColor() )
-			[
-				MapListView.ToSharedRef()
-			]
-		]
-	;
-}
-END_SLATE_FUNCTION_BUILD_OPTIMIZATION
-
-
-void FCustomProfileMarshaller::SetCook( bool bCook )
-{
-	bShouldCook = bCook;
-	RefreshContentScheme();
-}
-
-bool FCustomProfileMarshaller::GetCook(ILauncherProfilePtr InProfile) const
-{
-	if (InProfile == nullptr)
-	{
-		InProfile = Profile;
-	}
-	return (InProfile->GetCookMode() != ELauncherProfileCookModes::DoNotCook);
-}
-
-
-
-void FCustomProfileMarshaller::SetMapsToCook(TArray<FString> MapsToCook)
-{
-	Profile->ClearCookedMaps(); 
-	for (const FString& Map : MapsToCook)
-	{
-		Profile->AddCookedMap(Map);
-	}
-
-	OnPropertyChanged();
-}
-
-TArray<FString> FCustomProfileMarshaller::GetMapsToCook() const
-{
-	return Profile->GetCookedMaps();
-}
-
-float FCustomProfileMarshaller::GetMapListHeight() const
-{
-	return MapListHeight;
-}
-
-void FCustomProfileMarshaller::SetMapListHeight( float NewHeight )
-{
-	static const float MinMapListHeight = 100.0f;
-
-	MapListHeight = FMath::Max(NewHeight, MinMapListHeight);
-}
-
-FCustomProfileMarshaller::EMapOption FCustomProfileMarshaller::GetMapOption() const
-{
-	return MapOption;
-}
-
-void FCustomProfileMarshaller::SetMapOption( EMapOption NewMapOption )
-{
-	bool bShow = (NewMapOption == EMapOption::Selected);
-
-	MapOption = NewMapOption;
-
-	if (bShow)
-	{
-		// restore the cooked maps again, if any
-		if (CachedMapsToCook.Num() > 0 && Profile->GetCookedMaps().Num() == 0)
-		{
-			SetMapsToCook(CachedMapsToCook);
-			CachedMapsToCook.Reset();
-		}
-	}
-	else
-	{
-		// to set the 'cook startup maps only', its necessary to remove all the cooked maps - take a copy of the values to allow it to be restored
-		CachedMapsToCook = Profile->GetCookedMaps();
-		SetMapsToCook(TArray<FString>());
-	}
-
-	OnPropertyChanged();
-
-	MapListView->RefreshMapList();
-}
-
-
-
-
-BEGIN_SLATE_FUNCTION_BUILD_OPTIMIZATION
-TSharedRef<SWidget> FCustomProfileMarshaller::CreateDeployDeviceWidget()
-{
-	return SNew(SVerticalBox)
-
-		// device picker options
-		+SVerticalBox::Slot()
-		.AutoHeight()
-		.Padding(0,2)
-		[
-			SNew(SHorizontalBox)
-
-			+SHorizontalBox::Slot()
-			.Padding(0, 0)
-			.AutoWidth()
-			.VAlign(VAlign_Center)
-			[
-				SNew(SSegmentedControl<EDeployDeviceOption>)
-				.Value( this, &FCustomProfileMarshaller::GetDeployDeviceOption)
-				.OnValueChanged(this, &FCustomProfileMarshaller::SetDeployDeviceOption)
-
-				+SSegmentedControl<EDeployDeviceOption>::Slot(EDeployDeviceOption::Default)
-				.HAlign(HAlign_Center)
-				.VAlign(VAlign_Center)
-				[
-					SNew(STextBlock)
-					.Text(LOCTEXT("DefaultDeviceLabel", "Default Device"))
-					.Font(FCoreStyle::Get().GetFontStyle("SmallFont"))
-				]
-
-				+SSegmentedControl<EDeployDeviceOption>::Slot(EDeployDeviceOption::Selected)
-				.HAlign(HAlign_Center)
-				.VAlign(VAlign_Center)
-				[
-					SNew(STextBlock)
-					.Text(LOCTEXT("SelectedDevicesLabel", "Selected Devices"))
-					.Font(FCoreStyle::Get().GetFontStyle("SmallFont"))
-				]
-			]
-		]
-
-		// device picker list
-		+SVerticalBox::Slot()
-		.FillHeight(1)
-		[
-			SNew(SVerticalResizeBox)
-			.Visibility_Lambda([this](){ return (GetDeployDeviceOption() == EDeployDeviceOption::Selected) ? EVisibility::Visible : EVisibility::Collapsed; } )
-			.HandleHeight(4)
-			.ContentHeight(this, &FCustomProfileMarshaller::GetDeployDeviceListHeight)
-			.ContentHeightChanged(this, &FCustomProfileMarshaller::SetDeployDeviceListHeight)
-			.HandleColor( FAppStyle::Get().GetSlateColor("Colors.Secondary").GetSpecifiedColor() )
-			[
-				SAssignNew(DeployDeviceListView, SCustomLaunchDeviceListView)
-				.OnDeviceRemoved( this, &FCustomProfileMarshaller::OnDeviceRemoved)
-				.OnSelectionChanged(this, &FCustomProfileMarshaller::SetDeployDeviceIDs)
-				.SelectedDevices(this, &FCustomProfileMarshaller::GetDeployDeviceIDs)
-				.Platforms(this, &FCustomProfileMarshaller::GetSelectedPlatforms)
-			]
-		]
-	;
-}
-END_SLATE_FUNCTION_BUILD_OPTIMIZATION
-
-
-void FCustomProfileMarshaller::SetDeployDeviceIDs(TArray<FString> DeployDeviceIDs)
-{
-	if (Profile->GetDeployedDeviceGroup() == nullptr)
-	{
-		Profile->SetDeployedDeviceGroup(Model->GetProfileManager()->AddNewDeviceGroup());
-	}
-
-	Profile->GetDeployedDeviceGroup()->RemoveAllDevices(); 
-	for ( const FString& DeviceID : DeployDeviceIDs )
-	{
-		Profile->GetDeployedDeviceGroup()->AddDevice(DeviceID);
-	}
-	// Model->UpdatedCookedPlatformsFromDeployDeviceProxy(Profile); // fixme: needed?
-
-	OnPropertyChanged();
-}
-
-TArray<FString> FCustomProfileMarshaller::GetDeployDeviceIDs() const
-{
-	return Profile->GetDeployedDeviceGroup()->GetDeviceIDs();
-}
-
-float FCustomProfileMarshaller::GetDeployDeviceListHeight() const
-{
-	return DeployDeviceListHeight;
-}
-
-void FCustomProfileMarshaller::SetDeployDeviceListHeight( float NewHeight )
-{
-	static const float MinDeployDeviceListHeight = 100.0f;
-
-	DeployDeviceListHeight = FMath::Max(NewHeight, MinDeployDeviceListHeight);
-}
-
-void FCustomProfileMarshaller::OnDeviceRemoved(FString DeviceID)
+namespace CustomLaunchUI
 {
-	CachedDeployDeviceIDs.Remove(DeviceID);
-}
-
-FCustomProfileMarshaller::EDeployDeviceOption FCustomProfileMarshaller::GetDeployDeviceOption() const
-{
-	return DeployDeviceOption;
-}
-
-void FCustomProfileMarshaller::SetDeployDeviceOption( EDeployDeviceOption NewDeployDeviceOption )
-{
-	bool bShow = (NewDeployDeviceOption == EDeployDeviceOption::Selected);
-
-	DeployDeviceOption = NewDeployDeviceOption;
-
-	if (bShow)
-	{
-		// restore the deployed device list again, if any
-		if (CachedDeployDeviceIDs.Num() > 0 && Profile->GetDeployedDeviceGroup()->GetDeviceIDs().Num() == 0)
-		{
-			SetDeployDeviceIDs(Profile->GetDeployedDeviceGroup()->GetDeviceIDs());
-			CachedDeployDeviceIDs.Reset();
-		}
-	}
-	else
-	{
-		// to set the 'default' deploy option, its necessary to remove all the devices - take a copy of the values to allow it to be restored
-		CachedDeployDeviceIDs = Profile->GetDeployedDeviceGroup()->GetDeviceIDs();
-		SetDeployDeviceIDs(TArray<FString>());
-	}
-
-	OnPropertyChanged();
-
-	DeployDeviceListView->RefreshDeviceList();
-}
-
-
-void FCustomProfileMarshaller::SetBuild( bool bBuild )
-{
-	if (!bBuild)
-	{
-		Profile->SetBuildMode(ELauncherProfileBuildModes::DoNotBuild);
-	}
-	else if (GetForceBuild(Profile))
-	{
-		Profile->SetBuildMode(ELauncherProfileBuildModes::Build);
-	}
-	else
-	{
-		Profile->SetBuildMode(ELauncherProfileBuildModes::Auto);
-	}
-
-	OnPropertyChanged();
-}
-
-bool FCustomProfileMarshaller::GetBuild(ILauncherProfilePtr InProfile) const
-{
-	if (InProfile == nullptr)
-	{
-		InProfile = Profile;
-	}
-	return (InProfile->GetBuildMode() != ELauncherProfileBuildModes::DoNotBuild);
-}
-
-void FCustomProfileMarshaller::SetForceBuild( bool bForceBuild )
-{
-	if (!GetBuild())
-	{
-		Profile->SetBuildMode(ELauncherProfileBuildModes::DoNotBuild);
-	}
-	else if (bForceBuild)
-	{
-		Profile->SetBuildMode(ELauncherProfileBuildModes::Build);
-	}
-	else
-	{
-		Profile->SetBuildMode(ELauncherProfileBuildModes::Auto);
-	}
-
-	OnPropertyChanged();
-}
-
-bool FCustomProfileMarshaller::GetForceBuild(ILauncherProfilePtr InProfile) const
-{
-	if (InProfile == nullptr)
-	{
-		InProfile = Profile;
-	}
-	return (InProfile->GetBuildMode() == ELauncherProfileBuildModes::Build);
-}
-
-void FCustomProfileMarshaller::SetArchitecture( FString Architecture )
-{
-	// clear existing architectures
-	TArray<FString> Empty;
-	Profile->SetServerArchitectures(Empty);
-	Profile->SetEditorArchitectures(Empty);
-	Profile->SetClientArchitectures(Empty);
-
-	// set new single architecture
-	if (!Architecture.IsEmpty())
-	{
-		TArray<FString> Single;
-		Single.Add(Architecture);
-
-		switch (CachedBuildTargetType)
-		{
-			case EBuildTargetType::Server: Profile->SetServerArchitectures(Single); break;
-			case EBuildTargetType::Editor: Profile->SetEditorArchitectures(Single); break;
-			default:                       Profile->SetClientArchitectures(Single); break;
-		}
-	}
-
-	OnPropertyChanged();
-}
-
-FString FCustomProfileMarshaller::GetArchitecture() const
-{
-	TArray<FString> Architectures;
-
-	switch (CachedBuildTargetType)
-	{
-		case EBuildTargetType::Server: Architectures = Profile->GetServerArchitectures(); break;
-		case EBuildTargetType::Editor: Architectures = Profile->GetEditorArchitectures(); break;
-		default:                       Architectures = Profile->GetClientArchitectures(); break;
-	}
-
-	return (Architectures.Num() > 0) ? Architectures[0] : FString();
-}
-
-FText FCustomProfileMarshaller::GetArchitectureDisplayName( FString Architecture )
-{
-	if (Architecture.IsEmpty())
-	{
-		return LOCTEXT("DefaultArchName", "Project Default");
-	}
-	else if (Architecture == FPlatformMisc::GetHostArchitecture() && Model->IsHostPlatform(Profile))
+	FCustomProfileMarshaller::FCustomProfileMarshaller( const ILauncherProfileRef& InProfile, const TSharedRef<FModel>& InModel )
+		: FGenericProfileMarshaller( InProfile, InModel->GetDefaultCustomLaunchProfile(), InModel )
 	{
-		return FText::Format(LOCTEXT("HostArchLabel", "{0} (this platform)"), FText::FromString(Architecture) );
 	}
-	else
-	{
-		return FText::FromString(Architecture);
-	}
-}
-
 
-
-void FCustomProfileMarshaller::SetDeployToDevice( bool bDeployToDevice )
-{
-	if (ContentScheme != CustomLaunchUI::EContentScheme::CookOnTheFly)
+	void FCustomProfileMarshaller::Construct()
 	{
-		Profile->SetDeploymentMode(bDeployToDevice ? ELauncherProfileDeploymentModes::CopyToDevice : ELauncherProfileDeploymentModes::DoNotDeploy );
-		OnPropertyChanged();
-	}
-}
-
-bool FCustomProfileMarshaller::GetDeployToDevice(ILauncherProfilePtr InProfile) const
-{
-	if (InProfile == nullptr)
-	{
-		InProfile = Profile;
-	}
-	return InProfile->GetDeploymentMode() != ELauncherProfileDeploymentModes::DoNotDeploy;
-}
+		FGenericProfileMarshaller::Construct();
 
+		FLaunchProfileTreeNode& GeneralSettingsHeader = TreeData->AddHeading( LOCTEXT("GeneralSettingsHeading", "General Settings") );
+		AddProjectProperty(GeneralSettingsHeader);
+		AddTargetProperty(GeneralSettingsHeader);
+		AddPlatformProperty(GeneralSettingsHeader);
+		AddConfigurationProperty(GeneralSettingsHeader);
+		AddContentSchemeProperty(GeneralSettingsHeader);
 
-void FCustomProfileMarshaller::SetIsRunning( bool bRun )
-{
-	Profile->SetLaunchMode(bRun ? ELauncherProfileLaunchModes::DefaultRole : ELauncherProfileLaunchModes::DoNotLaunch);
-}
+		FLaunchProfileTreeNode& ContentSchemeHeader = TreeData->AddHeading( LOCTEXT("ContentSchemeHeading", "Content Scheme") );
+		AddCompressPakFilesProperty(ContentSchemeHeader);
+		AddUseIoStoreProperty(ContentSchemeHeader);
+		AddGenerateChunksProperty(ContentSchemeHeader);
+		AddImportZenSnapshotProperty(ContentSchemeHeader);
+		AddZenPakStreamingPathProperty(ContentSchemeHeader);
+		AddIncrementalCookProperty(ContentSchemeHeader);
 
-bool FCustomProfileMarshaller::GetIsRunning( ILauncherProfilePtr InProfile) const
-{
-	if (InProfile == nullptr)
-	{
-		InProfile = Profile;
-	}
-	
-	return InProfile->GetLaunchMode() != ELauncherProfileLaunchModes::DoNotLaunch;
-}
+		FLaunchProfileTreeNode& MapsAndCookingHeader = TreeData->AddHeading( LOCTEXT("CookingHeading", "Maps And Cooking") );
+		AddCookProperty(MapsAndCookingHeader);
+		AddMapsToCookProperty(MapsAndCookingHeader);
+		AddAdditionalCookerOptionsProperty(MapsAndCookingHeader);
 
-void FCustomProfileMarshaller::OnInitialMapChanged( TSharedPtr<FString> InitialMap, ESelectInfo::Type )
-{
-	if (InitialMap.IsValid())
-	{
-		Profile->GetDefaultLaunchRole()->SetInitialMap(*InitialMap);
-	}
-	else
-	{
-		Profile->GetDefaultLaunchRole()->SetInitialMap(FString());
-	}
-
-	OnPropertyChanged();
-}
+		FLaunchProfileTreeNode& BuildHeader = TreeData->AddHeading( LOCTEXT("BuildHeading", "Build") );
+		AddBuildProperty(BuildHeader);
+		AddForceBuildProperty(BuildHeader);
+		AddArchitectureProperty(BuildHeader);
 
-TSharedPtr<FString> FCustomProfileMarshaller::GetInitialMap() const
-{
-	FString InitialMap = Profile->GetDefaultLaunchRole()->GetInitialMap();
-	const TSharedPtr<FString>* FoundMapPtr = CachedStartupMaps.FindByPredicate( [InitialMap]( const TSharedPtr<FString>& Map )
-	{
-		return Map.IsValid() && (InitialMap == *Map);
-	});
+		FLaunchProfileTreeNode& DirectoryHeader = TreeData->AddHeading( LOCTEXT("DirectoryHeading", "Directory") );
+		AddArchiveBuildProperty(DirectoryHeader);
+		AddArchiveBuildDirectoryProperty(DirectoryHeader);
 
-	if (FoundMapPtr)
-	{
-		return *FoundMapPtr;
+		FLaunchProfileTreeNode& DeployAndRunHeader = TreeData->AddHeading( LOCTEXT("DeployAndRunHeading", "Deploy And Run") );
+		AddDeployProperty(DeployAndRunHeader);
+		AddTargetDeviceProperty(DeployAndRunHeader);
+		AddRunProperty(DeployAndRunHeader);
+		AddInitialMapProperty(DeployAndRunHeader);
+		AddCommandLineProperty(DeployAndRunHeader);
 	}
-	else
-	{
-		return MakeShared<FString>();
-	}
-}
-
 
 
 
-
-
-TSharedPtr<ILaunchProfileMarshaller> FCustomProfileMarshallerFactory::TryCreateMarshaller( const ILauncherProfileRef& InProfile, const TSharedRef<CustomLaunchUI::FModel>& InModel )
-{
-	if (!InModel->IsBasicLaunchProfile(InProfile) && !InModel->IsAdvancedProfile(InProfile))
+	TSharedPtr<ILaunchProfileMarshaller> FCustomProfileMarshallerFactory::TryCreateMarshaller( const ILauncherProfileRef& InProfile, const TSharedRef<FModel>& InModel )
 	{
 		return MakeShared<FCustomProfileMarshaller>(InProfile, InModel);
 	}
 
-	return nullptr;
-}
-
 
+}
 
 #undef LOCTEXT_NAMESPACE
diff --git a/Engine/Plugins/Developer/CustomLaunchUI/Source/Private/Marshalling/CustomProfileMarshaller.h b/Engine/Plugins/Developer/CustomLaunchUI/Source/Private/Marshalling/CustomProfileMarshaller.h
deleted file mode 100644
diff --git a/Engine/Plugins/Developer/CustomLaunchUI/Source/Private/Marshalling/GenericProfileMarshaller.cpp b/Engine/Plugins/Developer/CustomLaunchUI/Source/Private/Marshalling/GenericProfileMarshaller.cpp
new file mode 100644
--- /dev/null
+++ b/Engine/Plugins/Developer/CustomLaunchUI/Source/Private/Marshalling/GenericProfileMarshaller.cpp
@@ -0,0 +1,1304 @@
+// Copyright Epic Games, Inc. All Rights Reserved.
+
+#include "Marshalling/GenericProfileMarshaller.h"
+#include "Widgets/Shared/SCustomLaunchProjectCombo.h"
+#include "Widgets/Shared/SCustomLaunchBuildTargetCombo.h"
+#include "Widgets/Shared/SCustomLaunchPlatformCombo.h"
+#include "Widgets/Shared/SCustomLaunchContentSchemeCombo.h"
+#include "Widgets/Shared/SCustomLaunchDeviceListView.h"
+#include "Widgets/Shared/SCustomLaunchMapListView.h"
+#include "Widgets/Shared/SCustomLaunchCombo.h"
+#include "Widgets/Input/SSegmentedControl.h"
+#include "Widgets/Text/STextBlock.h"
+#include "Widgets/Layout/SBox.h"
+#include "Widgets/SBoxPanel.h"
+#include "SResizeBox.h"
+#include "SSearchableComboBox.h"
+#include "PlatformInfo.h"
+#include "IDesktopPlatform.h"
+#include "Interfaces/ITargetPlatformManagerModule.h"
+#include "Interfaces/ITargetPlatform.h"
+#include "GameProjectHelper.h"
+
+
+
+#define LOCTEXT_NAMESPACE "CustomProfileMarshaller"
+
+
+namespace CustomLaunchUI
+{
+	FGenericProfileMarshaller::FGenericProfileMarshaller( const ILauncherProfileRef& InProfile, const ILauncherProfileRef& InDefaultProfile, const TSharedRef<FModel>& InModel )
+		: TreeData( MakeShared<FLaunchProfileTreeData>(InProfile, InModel, this) )
+		, Profile(InProfile)
+		, DefaultProfile(InDefaultProfile)
+		, Model(InModel)
+	{
+		ProfileType = Model->GetProfileType(Profile);
+
+		ForPak = [this]()
+		{
+			return (ContentScheme == EContentScheme::PakFiles || ContentScheme == EContentScheme::DevelopmentPackage);
+		};
+
+		ForZenWS = [this]()
+		{
+			return (ContentScheme == EContentScheme::ZenPakStreaming);
+		};
+
+		ForCooked = [this]()
+		{
+			return (ContentScheme != EContentScheme::ZenPakStreaming && ContentScheme != EContentScheme::CookOnTheFly);
+		};
+
+		ForEnabledCooked = [this]()
+		{
+			return (bShouldCook && ContentScheme != EContentScheme::ZenPakStreaming && ContentScheme != EContentScheme::CookOnTheFly);
+		};
+
+		ForCode = [this]()
+		{
+			return GetBuild();
+		};
+
+		ForDeployment = [this]()
+		{
+			return GetDeployToDevice();
+		};
+
+		ForRun = [this]()
+		{
+			return GetIsRunning();
+		};
+
+		EmptyString = []()
+		{
+			return FString();
+		};
+
+	}
+
+	BEGIN_SLATE_FUNCTION_BUILD_OPTIMIZATION
+	void FGenericProfileMarshaller::Construct()
+	{
+		MapOption = Profile->GetCookedMaps().Num() > 0 ? EMapOption::Selected : EMapOption::Startup;
+		DeployDeviceOption = Profile->GetDeployedDeviceGroup()->GetDeviceIDs().Num() > 0 ? EDeployDeviceOption::Selected : EDeployDeviceOption::Default;
+		bShouldCook = GetCook();
+		ContentScheme = Model->DetermineProfileContentScheme(Profile);
+		CachedBuildTargetType = Model->GetBuildTargetInfo(Profile).Type;
+	
+		CacheStartupMapList();
+		CacheArchitectures();
+	}
+	END_SLATE_FUNCTION_BUILD_OPTIMIZATION
+
+
+
+
+
+
+
+
+
+	void FGenericProfileMarshaller::AddProjectProperty( FLaunchProfileTreeNode& HeadingNode )
+	{
+		if (ProfileType == EProfileType::Basic)
+		{
+			HeadingNode.AddWidget( LOCTEXT("ProjectLabel", "Project"), 
+				SNew(SCustomLaunchProjectCombo)
+				.OnSelectionChanged(this, &FGenericProfileMarshaller::SetProjectName)
+				.SelectedProject(this, &FGenericProfileMarshaller::GetProjectPath)
+				.HasProject(this, &FGenericProfileMarshaller::HasProject)
+				.CurrentProjectOption(SCustomLaunchProjectCombo::ECurrentProjectOption::Empty)
+				.Font(FCoreStyle::Get().GetFontStyle("SmallFont"))
+			);
+		}
+		else
+		{
+			HeadingNode.AddWidget( LOCTEXT("ProjectLabel", "Project"), 
+				SNew(SCustomLaunchProjectCombo)
+				.OnSelectionChanged(this, &FGenericProfileMarshaller::SetProjectName)
+				.SelectedProject(this, &FGenericProfileMarshaller::GetProjectPath)
+				.HasProject(this, &FGenericProfileMarshaller::HasProject)
+				.ShowAnyProjectOption(true)
+				.CurrentProjectOption(SCustomLaunchProjectCombo::ECurrentProjectOption::ActualProject)
+				.Font(FCoreStyle::Get().GetFontStyle("SmallFont"))
+			);
+		}
+	}
+
+	void FGenericProfileMarshaller::AddTargetProperty( FLaunchProfileTreeNode& HeadingNode )
+	{
+		if (ProfileType == EProfileType::Basic)
+		{
+			HeadingNode.AddWidget( LOCTEXT("TargetLabel", "Target"), 
+				{
+					.IsDefault = [this]()		{ return !Profile->HasBuildTargetSpecified() || Profile->GetBuildTarget().IsEmpty(); },
+					.SetToDefault = [this]()	{ SetBuildTarget(FString()); },
+				},
+				SNew(SCustomLaunchBuildTargetCombo)
+				.OnSelectionChanged(this, &FGenericProfileMarshaller::SetBuildTarget)
+				.SelectedBuildTarget(this, &FGenericProfileMarshaller::GetBuildTarget)
+				.SelectedProject(this, &FGenericProfileMarshaller::GetProjectPath)
+				.Font(FCoreStyle::Get().GetFontStyle("SmallFont"))
+			);
+		}
+		else
+		{
+			HeadingNode.AddWidget( LOCTEXT("TargetLabel", "Target"), 
+				{
+					.IsDefault = [this]()		{ return !Profile->HasBuildTargetSpecified() || Profile->GetBuildTarget().IsEmpty(); },
+					.SetToDefault = [this]()	{ SetBuildTarget(FString()); },
+					.IsEnabled = [this]()		{ return Profile->HasProjectSpecified(); },
+				},
+				SNew(SCustomLaunchBuildTargetCombo)
+				.OnSelectionChanged(this, &FGenericProfileMarshaller::SetBuildTarget)
+				.SelectedBuildTarget(this, &FGenericProfileMarshaller::GetBuildTarget)
+				.SupportedTargetTypes(this, &FGenericProfileMarshaller::GetSupportedBuildTargetTypes)
+				.SelectedProject(this, &FGenericProfileMarshaller::GetProjectPath)
+				.Font(FCoreStyle::Get().GetFontStyle("SmallFont"))
+			);
+		}
+	}
+
+	void FGenericProfileMarshaller::AddPlatformProperty( FLaunchProfileTreeNode& HeadingNode )
+	{
+		HeadingNode.AddWidget( LOCTEXT("PlatformLabel", "Platform"), 
+			SNew(SCustomLaunchPlatformCombo)
+			.SelectedPlatforms(this, &FGenericProfileMarshaller::GetSelectedPlatforms)
+			.OnSelectionChanged(this, &FGenericProfileMarshaller::SetSelectedPlatforms)
+			.BasicPlatformsOnly(bUseFriendlyBuildTargetSelection)
+			.Font(FCoreStyle::Get().GetFontStyle("SmallFont"))
+		);
+	}
+
+	void FGenericProfileMarshaller::AddConfigurationProperty( FLaunchProfileTreeNode& HeadingNode )
+	{
+		HeadingNode.AddWidget( LOCTEXT("ConfigurationLabel", "Configuration"), 
+			SNew(SCustomLaunchLexToStringCombo<EBuildConfiguration>)
+			.OnSelectionChanged(this, &FGenericProfileMarshaller::SetBuildConfiguration)
+			.SelectedItem(this, &FGenericProfileMarshaller::GetBuildConfiguration)
+			.Items( TArray<EBuildConfiguration>({ EBuildConfiguration::Debug, EBuildConfiguration::Development, EBuildConfiguration::Test, EBuildConfiguration::Shipping }) )
+		);
+	}
+
+	void FGenericProfileMarshaller::AddContentSchemeProperty( FLaunchProfileTreeNode& HeadingNode )
+	{
+		HeadingNode.AddWidget( LOCTEXT("ContentSchemeLabel", "Content Scheme"), 
+			SNew(SCustomLaunchContentSchemeCombo)
+			.OnSelectionChanged(this, &FGenericProfileMarshaller::SetContentScheme)
+			.SelectedContentScheme_Lambda( [this]() { return ContentScheme; } )
+			.IsContentSchemeAvailable(this, &FGenericProfileMarshaller::IsContentSchemeAvailable)
+			.Font(FCoreStyle::Get().GetFontStyle("SmallFont"))
+		);
+	}
+
+	void FGenericProfileMarshaller::AddCompressPakFilesProperty( FLaunchProfileTreeNode& HeadingNode )
+	{
+		HeadingNode.AddBoolean( LOCTEXT("CompressPakFilesLabel", "Compress Pak Files"),
+			{
+				.GetValue = [this]()			{ return Profile->IsCompressed(); },
+				.SetValue = [this](bool bValue)	{ Profile->SetCompressed(bValue); },
+				.GetDefaultValue = [this]()		{ return DefaultProfile->IsCompressed(); },
+				.IsVisible = ForPak,
+			}
+		);
+	}
+
+	void FGenericProfileMarshaller::AddUseIoStoreProperty( FLaunchProfileTreeNode& HeadingNode )
+	{
+		HeadingNode.AddBoolean( LOCTEXT("UseIoStoreLabel", "Use Io Store"),
+			{
+				.GetValue = [this]()			{ return Profile->IsUsingIoStore(); },
+				.SetValue = [this](bool bValue)	{ Profile->SetUseIoStore(bValue); },
+				.GetDefaultValue = [this]()		{ return DefaultProfile->IsUsingIoStore(); },
+				.IsVisible = ForPak,
+			}
+		);
+	}
+
+	void FGenericProfileMarshaller::AddGenerateChunksProperty( FLaunchProfileTreeNode& HeadingNode )
+	{
+		HeadingNode.AddBoolean( LOCTEXT("GenerateChunksLabel", "Generate Chunks"),
+			{
+				.GetValue = [this]()			{ return Profile->IsGeneratingChunks(); },
+				.SetValue = [this](bool bValue)	{ Profile->SetGenerateChunks(bValue); },
+				.GetDefaultValue = [this]()		{ return DefaultProfile->IsGeneratingChunks(); },
+				.IsVisible = ForPak,
+			}
+		);
+	}
+
+	void FGenericProfileMarshaller::AddImportZenSnapshotProperty( FLaunchProfileTreeNode& HeadingNode )
+	{
+		HeadingNode.AddBoolean( LOCTEXT("ImportZenSnapshotLabel", "Import Best Match Zen Snapshot"),
+			{
+				.GetValue = [this]()			{ return Profile->IsImportingZenSnapshot(); },
+				.SetValue = [this](bool bValue)	{ Profile->SetImportingZenSnapshot(bValue); },
+				.GetDefaultValue = [this]()		{ return DefaultProfile->IsImportingZenSnapshot(); },
+				.IsVisible = ForCooked,
+			}
+		);
+	}
+
+	void FGenericProfileMarshaller::AddZenPakStreamingPathProperty( FLaunchProfileTreeNode& HeadingNode )
+	{
+		HeadingNode.AddDirectoryString( LOCTEXT("ZenPakStreamingPathLabel", "Zen Pak Streaming Path"),
+			{
+				.GetValue = [this]()				{ return Profile->GetZenPakStreamingPath(); },
+				.SetValue = [this](FString Value)	{ Profile->SetZenPakStreamingPath(Value); },
+				.GetDefaultValue = EmptyString,
+				.IsVisible = ForZenWS,
+			}
+		);
+	}
+
+	void FGenericProfileMarshaller::AddIncrementalCookProperty( FLaunchProfileTreeNode& HeadingNode )
+	{
+		HeadingNode.AddBoolean( LOCTEXT("IncrementalCookLabel", "Incremental Cook"),
+			{
+				.GetValue = [this]()			{ return Profile->IsCookingIncrementally() || Profile->IsImportingZenSnapshot(); }, // show it as ticked when snapshot import is enabled for now to indicate that it will be doing something incremental. eventually this is likely to become a separate field
+				.SetValue = [this](bool bValue)	{ Profile->SetIncrementalCooking(bValue); },
+				.GetDefaultValue = [this]()		{ return DefaultProfile->IsCookingIncrementally(); },
+				.IsVisible = ForCooked,
+				.IsEnabled = [this]()			{ return !Profile->IsImportingZenSnapshot(); }, // zen snapshot import & incremental cooking can't be used together. '-iterativecook' is the 'old' system, and zen import will enable the new '-cookincremental' & the two flags don't play nicely together
+			}
+		);
+	}
+
+
+
+	void FGenericProfileMarshaller::AddCookProperty( FLaunchProfileTreeNode& HeadingNode )
+	{
+		HeadingNode.AddBoolean( LOCTEXT("CookLabel", "Cook Content"),
+			{
+				.GetValue = [this]()			{ return GetCook(); },
+				.SetValue = [this](bool bValue)	{ SetCook(bValue); },
+				.GetDefaultValue = [this]()		{ return GetCook(DefaultProfile); },
+				.IsVisible = ForCooked
+			}
+		);
+	}
+
+	void FGenericProfileMarshaller::AddMapsToCookProperty( FLaunchProfileTreeNode& HeadingNode )
+	{
+		HeadingNode.AddWidget( LOCTEXT("MapsToCookLabel", "Maps To Cook"), 
+			{
+				.IsVisible = ForEnabledCooked,
+			},
+			CreateMapListWidget()
+		);
+	}
+
+	void FGenericProfileMarshaller::AddAdditionalCookerOptionsProperty( FLaunchProfileTreeNode& HeadingNode )
+	{
+		HeadingNode.AddString( LOCTEXT("AdditionalCookerOptionsLabel", "Additional Cooker Options"), 
+			{
+				.GetValue = [this]()				{ return Profile->GetCookOptions(); },
+				.SetValue = [this](FString Value)	{ Profile->SetCookOptions( Value ); },
+				.GetDefaultValue = EmptyString,
+				.IsVisible = ForEnabledCooked,
+			}
+		);
+	}
+
+	void FGenericProfileMarshaller::AddBuildProperty( FLaunchProfileTreeNode& HeadingNode )
+	{
+		HeadingNode.AddBoolean( LOCTEXT("BuildLabel", "Build the game"),
+			{
+				.GetValue = [this]()			{ return GetBuild(); },
+				.SetValue = [this](bool bValue)	{ SetBuild(bValue); },
+				.GetDefaultValue = [this]()		{ return GetBuild(DefaultProfile); },
+			}
+		);
+	}
+
+	void FGenericProfileMarshaller::AddForceBuildProperty( FLaunchProfileTreeNode& HeadingNode )
+	{
+		HeadingNode.AddBoolean( LOCTEXT("ForceBuildLabel", "Build even if a pre-built target exists"),
+			{
+				.GetValue = [this]()			{ return GetForceBuild(); },
+				.SetValue = [this](bool bValue)	{ SetForceBuild(bValue); },
+				.GetDefaultValue = [this]()		{ return GetForceBuild(DefaultProfile); },
+				.IsEnabled = ForCode,
+			}
+		);
+	}
+
+	void FGenericProfileMarshaller::AddBuidUATProperty( FLaunchProfileTreeNode& HeadingNode ) // fixme: not really relevant
+	{
+		HeadingNode.AddBoolean( LOCTEXT("BuildUATLabel", "Build UAT"),
+			{
+				.GetValue = [this]()			{ return Profile->IsBuildingUAT(); },
+				.SetValue = [this](bool bValue)	{ Profile->SetBuildUAT(bValue); },
+				.GetDefaultValue = [this]()		{ return DefaultProfile->IsBuildingUAT(); },
+				.IsEnabled = ForCode,
+			}
+		);
+	}
+
+	void FGenericProfileMarshaller::AddArchitectureProperty( FLaunchProfileTreeNode& HeadingNode )
+	{
+		HeadingNode.AddWidget( LOCTEXT("ArchitectureLabel", "Architecture"),
+			{
+				.IsDefault = [this]()		{ return GetArchitecture().IsEmpty(); },
+				.SetToDefault = [this]()	{ SetArchitecture(FString()); },
+				.IsVisible = [this]()       { return CachedArchitectures.Num() > 0; }
+			},
+			SNew(SCustomLaunchStringCombo)
+			.OnSelectionChanged( this, &FGenericProfileMarshaller::SetArchitecture )
+			.SelectedItem( this, &FGenericProfileMarshaller::GetArchitecture)
+			.GetDisplayName( this, &FGenericProfileMarshaller::GetArchitectureDisplayName )
+			.Items_Lambda( [this]() { return CachedArchitectures;} )
+		);
+	}
+
+
+
+
+	void FGenericProfileMarshaller::AddStagingDirectoryProperty( FLaunchProfileTreeNode& HeadingNode ) // fixme: not really relevant
+	{
+		HeadingNode.AddDirectoryString( LOCTEXT("CustomStagingPathLabel", "Custom Stage Directory"),
+			{
+				.GetValue = [this]()				{ return Profile->GetPackageDirectory(); },
+				.SetValue = [this](FString Value)	{ Profile->SetPackageDirectory(Value); },
+				.GetDefaultValue = EmptyString,
+				.IsVisible = ForCooked,
+			}
+		);
+	}
+
+	void FGenericProfileMarshaller::AddArchiveBuildProperty( FLaunchProfileTreeNode& HeadingNode )
+	{
+		HeadingNode.AddBoolean( LOCTEXT("ArchiveBuildLabel", "Archive Build"),
+			{
+				.GetValue = [this]()			{ return Profile->IsArchiving(); },
+				.SetValue = [this](bool bValue)	{ Profile->SetArchive(bValue); },
+				.GetDefaultValue = [this]()		{ return DefaultProfile->IsArchiving(); },
+				.IsVisible = ForCooked,
+			}
+		);
+	}
+
+
+	void FGenericProfileMarshaller::AddArchiveBuildDirectoryProperty( FLaunchProfileTreeNode& HeadingNode )
+	{
+		HeadingNode.AddDirectoryString( LOCTEXT("ArchivePathLabel", "Archive Directory"),
+			{
+				.GetValue = [this]()				{ return Profile->GetArchiveDirectory(); },
+				.SetValue = [this](FString Value)	{ Profile->SetArchiveDirectory(Value); },
+				.GetDefaultValue = EmptyString,
+				.IsVisible = ForCooked,				
+				.IsEnabled = [this]()				{ return Profile->IsArchiving(); }
+			}
+		);
+	}
+
+
+
+	void FGenericProfileMarshaller::AddDeployProperty( FLaunchProfileTreeNode& HeadingNode )
+	{
+		HeadingNode.AddBoolean( LOCTEXT("DeployLabel", "Deploy To Device"),
+			{
+				.GetValue = [this]()			{ return GetDeployToDevice(); },
+				.SetValue = [this](bool bValue)	{ SetDeployToDevice(bValue); },
+				.GetDefaultValue = [this]()		{ return GetDeployToDevice(DefaultProfile); },
+				.IsVisible = [this]()			{ return (ContentScheme != EContentScheme::CookOnTheFly); }
+			}
+		);
+	}
+
+	void FGenericProfileMarshaller::AddIncrementalDeployProperty( FLaunchProfileTreeNode& HeadingNode ) // fixme: not really relavent
+	{
+		HeadingNode.AddBoolean( LOCTEXT("DeployModifiedLabel", "Only Deploy Modified Content"),
+			{
+				.GetValue = [this]()			{ return Profile->IsDeployingIncrementally(); },
+				.SetValue = [this](bool bValue)	{ Profile->SetIncrementalDeploying(bValue); },
+				.GetDefaultValue = [this]()		{ return DefaultProfile->IsDeployingIncrementally(); },
+				.IsEnabled = ForDeployment,
+			}
+		);
+	}
+
+	void FGenericProfileMarshaller::AddTargetDeviceProperty( FLaunchProfileTreeNode& HeadingNode )
+	{
+		HeadingNode.AddWidget( LOCTEXT("TargetDeviceLabel", "Target Device"), 
+			{
+				.IsEnabled = ForDeployment,
+			},
+			CreateDeployDeviceWidget()
+		);
+	}
+
+
+	void FGenericProfileMarshaller::AddRunProperty( FLaunchProfileTreeNode& HeadingNode )
+	{
+		HeadingNode.AddBoolean( LOCTEXT("RunLabel", "Run"),
+			{
+				.GetValue = [this]()			{ return GetIsRunning(); },
+				.SetValue = [this](bool bValue)	{ SetIsRunning(bValue); },
+				.GetDefaultValue = [this]()		{ return GetIsRunning(DefaultProfile); },
+			}
+		);
+	}
+
+
+	void FGenericProfileMarshaller::AddInitialMapProperty( FLaunchProfileTreeNode& HeadingNode )
+	{
+		// todo: custom map picker 
+		HeadingNode.AddWidget( LOCTEXT("InitialMapLabel", "Initial Map"), 
+			{
+				.IsDefault = [this]()		{ return Profile->GetDefaultLaunchRole()->GetInitialMap().IsEmpty(); },
+				.SetToDefault = [this]()	{ Profile->GetDefaultLaunchRole()->SetInitialMap(FString()); InitalMapCombo->SetSelectedItem(GetInitialMap()); },
+				.IsEnabled = ForRun,
+			},
+			SAssignNew(InitalMapCombo, SSearchableComboBox)
+			.OptionsSource(&CachedStartupMaps)
+			.OnSelectionChanged( this, &FGenericProfileMarshaller::OnInitialMapChanged )
+			.OnGenerateWidget( this, &FGenericProfileMarshaller::OnGenerateComboWidget )
+			.Content()
+			[
+				SNew(STextBlock)
+				.Text_Lambda( [this]() { return FText::FromString( *GetInitialMap() ); } )
+				.Font(FCoreStyle::Get().GetFontStyle("SmallFont"))
+			]
+		);
+	}
+
+
+	void FGenericProfileMarshaller::AddCommandLineProperty( FLaunchProfileTreeNode& HeadingNode )
+	{
+		HeadingNode.AddString( LOCTEXT("CommandLineLabel", "Command Line"), 
+			{
+				.GetValue = [this]()				{ return GetCommandLine(); },
+				.SetValue = [this](FString Value)	{ SetCommandLine(Value); },
+				.GetDefaultValue = EmptyString,
+				.IsEnabled = ForRun
+			}
+		);
+	}
+
+
+
+
+
+
+
+
+
+
+
+
+
+
+
+
+
+
+
+
+
+
+
+	void FGenericProfileMarshaller::OnPropertyChanged()
+	{
+		if (ProfileType == EProfileType::Basic)
+		{
+			// do not save basic profiles - they're transient
+		}
+		else
+		{
+			Model->GetProfileManager()->SaveJSONProfile(Profile);
+		}
+	}
+
+
+	void FGenericProfileMarshaller::CacheStartupMapList()
+	{
+		CachedStartupMaps.Reset();
+		CachedStartupMaps.Add(MakeShared<FString>());
+		for ( const FString& Map : FGameProjectHelper::GetAvailableMaps(Profile->GetProjectBasePath(), Profile->SupportsEngineMaps(), true))
+		{
+			CachedStartupMaps.Add(MakeShared<FString>(Map));
+		}
+
+		if (InitalMapCombo.IsValid())
+		{
+			InitalMapCombo->SetSelectedItem( GetInitialMap(), ESelectInfo::Direct );
+			InitalMapCombo->RefreshOptions();
+		}
+	}
+
+	void FGenericProfileMarshaller::CacheArchitectures()
+	{
+		CachedArchitectures.Reset();
+
+		const PlatformInfo::FTargetPlatformInfo* PlatformInfo = Model->GetPlatformInfo(Profile);
+		if (PlatformInfo != nullptr)
+		{
+			const ITargetPlatform* TargetPlatform = GetTargetPlatformManager()->FindTargetPlatform(PlatformInfo->IniPlatformName);
+			if (TargetPlatform)
+			{
+				TargetPlatform->GetPossibleArchitectures(CachedArchitectures);
+				if (CachedArchitectures.Num() > 0)
+				{
+					CachedArchitectures.Insert(FString(), 0); // empty string for "project default" option
+				}
+			}
+		}
+	}
+
+
+
+	TSharedRef<SWidget> FGenericProfileMarshaller::OnGenerateComboWidget( TSharedPtr<FString> InComboString )
+	{
+		return SNew(STextBlock)
+			.Text(InComboString.IsValid() ? FText::FromString(*InComboString) : FText::GetEmpty() )
+			.Font(FCoreStyle::Get().GetFontStyle("SmallFont"));
+
+	}
+
+
+
+	void FGenericProfileMarshaller::SetSelectedPlatforms( TArray<FString> SelectedPlatforms )
+	{
+		Profile->ClearCookedPlatforms();
+
+		if (bUseFriendlyBuildTargetSelection)
+		{
+			FTargetInfo BuildTargetInfo = FModel::GetBuildTargetInfo(GetBuildTarget(), GetProjectPath());
+
+			for (const FString& Platform : SelectedPlatforms)
+			{
+				Profile->AddCookedPlatform( FModel::GetBuildTargetPlatformName(Platform, BuildTargetInfo));
+			}
+		}
+		else
+		{
+			for (const FString& Platform : SelectedPlatforms)
+			{
+				Profile->AddCookedPlatform(Platform);
+			}
+		}
+		OnPropertyChanged();
+
+		DeployDeviceListView->OnSelectedPlatformChanged();
+
+		CacheArchitectures();
+	}
+
+	TArray<FString> FGenericProfileMarshaller::GetSelectedPlatforms() const
+	{
+		TArray<FString> Platforms;
+
+		if (bUseFriendlyBuildTargetSelection)
+		{
+			for (const FString& Platform : Profile->GetCookedPlatforms())
+			{
+				Platforms.Add(FModel::GetVanillaPlatformName(Platform));
+			}
+		}
+		else
+		{
+			Platforms = Profile->GetCookedPlatforms();
+		}
+
+		return MoveTemp(Platforms);
+	}
+
+
+	FString FGenericProfileMarshaller::GetProjectPath() const
+	{
+		if (ProfileType == EProfileType::Basic)
+		{
+			if (Profile->HasProjectSpecified())
+			{
+				return Profile->GetProjectPath();
+			}
+			else
+			{
+				return Model->GetProfileManager()->GetProjectPath();
+			}
+		}
+		else
+		{
+			if (Profile->HasProjectSpecified())
+			{
+				return Profile->GetProjectPath();
+			}
+
+			return FString();
+		}
+	}
+
+	void FGenericProfileMarshaller::SetProjectName(FString ProjectPath)
+	{
+		Profile->SetProjectSpecified(!ProjectPath.IsEmpty());
+		Profile->SetProjectPath(ProjectPath);
+
+		if (bUseFriendlyBuildTargetSelection && ProfileType != EProfileType::Basic )
+		{
+			Model->UpdateCookedPlatformsFromBuildTarget(Profile);
+		}
+
+		OnPropertyChanged();
+
+		if (MapListView.IsValid())
+		{
+			MapListView->RefreshMapList();
+		}
+
+		CacheStartupMapList();
+	}
+
+	bool FGenericProfileMarshaller::HasProject() const
+	{
+		return Profile->HasProjectSpecified();
+	}
+
+
+	FString FGenericProfileMarshaller::GetBuildTarget() const
+	{
+		if (ProfileType == EProfileType::Basic)
+		{
+			if (Profile->HasBuildTargetSpecified())
+			{
+				return Profile->GetBuildTarget();
+			}
+			else
+			{
+				return Model->GetProfileManager()->GetBuildTarget();
+			}
+		}
+		else
+		{
+			return Profile->GetBuildTarget();
+		}
+	}
+
+	void FGenericProfileMarshaller::SetBuildTarget(FString BuildTarget)
+	{
+		Profile->SetBuildTargetSpecified(!BuildTarget.IsEmpty());
+		Profile->SetBuildTarget(BuildTarget);
+	
+		if (ProfileType == EProfileType::Basic )
+		{
+			Model->UpdatedCookedPlatformsFromDeployDeviceProxy(Profile);
+		}
+		else if (bUseFriendlyBuildTargetSelection)
+		{
+			Model->UpdateCookedPlatformsFromBuildTarget(Profile);
+		}
+
+		CachedBuildTargetType = Model->GetBuildTargetInfo(Profile).Type;
+		OnPropertyChanged();
+	}
+
+	TArray<EBuildTargetType> FGenericProfileMarshaller::GetSupportedBuildTargetTypes() const
+	{
+		TArray<EBuildTargetType> Result;
+
+		if (Profile->GetCookedPlatforms().Num() > 0)
+		{
+			const PlatformInfo::FTargetPlatformInfo* PlatformInfo = PlatformInfo::FindPlatformInfo(FName(Profile->GetCookedPlatforms()[0]));
+			if (PlatformInfo)
+			{
+				if (bUseFriendlyBuildTargetSelection)
+				{
+					Result.Add(PlatformInfo->VanillaInfo->PlatformType);
+
+					for (const PlatformInfo::FTargetPlatformInfo* PlatformFlavorInfo : PlatformInfo->VanillaInfo->Flavors)
+					{
+						Result.AddUnique(PlatformFlavorInfo->PlatformType);
+					}
+				}
+				else
+				{
+					Result.Add(PlatformInfo->PlatformType);
+				}
+			}
+		}
+
+		return MoveTemp(Result);
+	}
+
+	void FGenericProfileMarshaller::SetBuildConfiguration(EBuildConfiguration BuildConfiguration)
+	{
+		Profile->SetBuildConfiguration(BuildConfiguration);
+		OnPropertyChanged();
+	}
+
+	EBuildConfiguration FGenericProfileMarshaller::GetBuildConfiguration() const
+	{
+		return Profile->GetBuildConfiguration();
+	}
+
+
+
+	void FGenericProfileMarshaller::RefreshContentScheme()
+	{
+		EContentScheme CurrentContentScheme = Model->DetermineProfileContentScheme(Profile);
+		SetContentScheme(CurrentContentScheme);
+	}
+
+	void FGenericProfileMarshaller::SetContentScheme(EContentScheme InContentScheme)
+	{
+		ELauncherProfileDeploymentModes::Type DeploymentMode = GetDeployToDevice() ? ELauncherProfileDeploymentModes::CopyToDevice : ELauncherProfileDeploymentModes::DoNotDeploy;
+		Model->SetProfileContentScheme(InContentScheme, Profile, bShouldCook, DeploymentMode );
+		OnPropertyChanged();
+
+		// refresh the cached content scheme in case the option that was selected is not available
+		ContentScheme = Model->DetermineProfileContentScheme(Profile);
+	}
+
+	bool FGenericProfileMarshaller::IsContentSchemeAvailable(EContentScheme InContentScheme, FText& OutReason) const
+	{
+		FProjectSettings ProjectSettings = Model->GetProjectSettings(Profile);
+
+		// basic launch is aimed at launching the current content, not an external build
+		if (ProfileType == EProfileType::Basic && InContentScheme == EContentScheme::ZenPakStreaming)
+		{
+			// don't set a reason - just hide the item for Basic Launch
+			return false;
+		}
+
+		// loose files can't be selected if the project is using zen store because there's no way to opt-out of Zen Store from the UAT command line.
+		if (InContentScheme == EContentScheme::LooseFiles && ProjectSettings.bUseZenStore)
+		{
+			OutReason = LOCTEXT("NoLooseFilesReason", "Loose Files cannot be used when Zen Store is enabled in Project Settings");
+			return false;
+		}
+
+		// don't show zen pak streaming option if it isn't going to be set up automatically by UAT
+		// @todo: could potentially look at shelling out to Zen & querying if we have a dynamic workspace? no support for async config queries in this tool yet though
+		if (InContentScheme == EContentScheme::ZenPakStreaming && !ProjectSettings.bHasAutomaticZenPakStreamingWorkspaceCreation)
+		{
+			OutReason = LOCTEXT("NoZenPakReason", "Automatic Zen Pak streaming workspace creation has not been enabled in Project Settings");
+			return false;
+		}
+	
+		// cannot launch via Zen if we're targeting a remote device and Zen isn't accepting external connections
+		if ((InContentScheme == EContentScheme::ZenStreaming || InContentScheme == EContentScheme::ZenPakStreaming) 
+			&& !ProjectSettings.bAllowRemoteNetworkService && !Model->IsHostPlatform(Profile))
+		{
+			OutReason = LOCTEXT("NoZenReason", "Zen Streaming to a remote device requires AllowRemoteNetworkService");
+			return false;
+		}
+
+		return true;
+	}
+
+
+
+	FString FGenericProfileMarshaller::GetCommandLine() const
+	{
+		// get unified command line from these two fields. the first is presented in "Build" for old Project Launcher, and the latter is presented in "Launch" for old Project Launcher)
+		// when we save back to the profile, this will be stored just in the "Build" one for clarity (because multiple roles are not supported in old or new project launcher)
+		FString CommandLine = Profile->GetAdditionalCommandLineParameters().TrimStartAndEnd() + TEXT(" ") + Profile->GetDefaultLaunchRole()->GetUATCommandLine().TrimStartAndEnd();
+		CommandLine.TrimStartAndEndInline();
+
+		return CommandLine;
+	}
+
+	void FGenericProfileMarshaller::SetCommandLine( const FString& NewCommandLine )
+	{
+		Profile->SetAdditionalCommandLineParameters(NewCommandLine);
+		Profile->GetDefaultLaunchRole()->SetCommandLine(TEXT(""));
+		OnPropertyChanged();
+	}
+
+
+
+
+	BEGIN_SLATE_FUNCTION_BUILD_OPTIMIZATION
+	TSharedRef<SWidget> FGenericProfileMarshaller::CreateMapListWidget()
+	{
+		MapListView = 
+			SNew(SCustomLaunchMapListView)
+			.OnSelectionChanged(this, &FGenericProfileMarshaller::SetMapsToCook)
+			.SelectedMaps(this, &FGenericProfileMarshaller::GetMapsToCook)
+			.ProjectPath(this, &FGenericProfileMarshaller::GetProjectPath)
+		;
+
+		return SNew(SVerticalBox)
+
+			// map option controls
+			+SVerticalBox::Slot()
+			.AutoHeight()
+			.Padding(0,2)
+			[
+				SNew(SHorizontalBox)
+
+				// map option
+				+SHorizontalBox::Slot()
+				.Padding(0, 0)
+				.AutoWidth()
+				.VAlign(VAlign_Center)
+				[
+					SNew(SSegmentedControl<EMapOption>)
+					.Value( this, &FGenericProfileMarshaller::GetMapOption)
+					.OnValueChanged(this, &FGenericProfileMarshaller::SetMapOption)
+
+					+SSegmentedControl<EMapOption>::Slot(EMapOption::Startup)
+					.HAlign(HAlign_Center)
+					.VAlign(VAlign_Center)
+					[
+						SNew(STextBlock)
+						.Text(LOCTEXT("StartupMapsLabel", "Startup Maps"))
+						.Font(FCoreStyle::Get().GetFontStyle("SmallFont"))
+					]
+
+					+SSegmentedControl<EMapOption>::Slot(EMapOption::Selected)
+					.HAlign(HAlign_Center)
+					.VAlign(VAlign_Center)
+					[
+						SNew(STextBlock)
+						.Text(LOCTEXT("SelectedMapsLabel", "Selected Maps"))
+						.Font(FCoreStyle::Get().GetFontStyle("SmallFont"))
+					]
+				]
+
+				// map selector controls (search etc)
+				+SHorizontalBox::Slot()
+				.Padding(8, 0)
+				.FillWidth(1)
+				.VAlign(VAlign_Center)
+				[
+					SNew(SBox)
+					.Visibility_Lambda([this](){ return GetMapOption() == EMapOption::Selected  ? EVisibility::Visible : EVisibility::Collapsed; } )
+					[
+						MapListView->MakeControlsWidget()
+					]
+				]
+			]
+
+			// map list
+			+SVerticalBox::Slot()
+			.FillHeight(1)
+			[
+				SNew(SVerticalResizeBox)
+				.Visibility_Lambda([this](){ return (GetMapOption() == EMapOption::Selected) ? EVisibility::Visible : EVisibility::Collapsed; } )
+				.HandleHeight(4)
+				.ContentHeight(this, &FGenericProfileMarshaller::GetMapListHeight)
+				.ContentHeightChanged(this, &FGenericProfileMarshaller::SetMapListHeight)
+				.HandleColor( FAppStyle::Get().GetSlateColor("Colors.Secondary").GetSpecifiedColor() )
+				[
+					MapListView.ToSharedRef()
+				]
+			]
+		;
+	}
+	END_SLATE_FUNCTION_BUILD_OPTIMIZATION
+
+
+	void FGenericProfileMarshaller::SetCook( bool bCook )
+	{
+		bShouldCook = bCook;
+		RefreshContentScheme();
+	}
+
+	bool FGenericProfileMarshaller::GetCook(ILauncherProfilePtr InProfile) const
+	{
+		if (InProfile == nullptr)
+		{
+			InProfile = Profile;
+		}
+		return (InProfile->GetCookMode() != ELauncherProfileCookModes::DoNotCook);
+	}
+
+
+
+	void FGenericProfileMarshaller::SetMapsToCook(TArray<FString> MapsToCook)
+	{
+		Profile->ClearCookedMaps(); 
+		for (const FString& Map : MapsToCook)
+		{
+			Profile->AddCookedMap(Map);
+		}
+
+		OnPropertyChanged();
+	}
+
+	TArray<FString> FGenericProfileMarshaller::GetMapsToCook() const
+	{
+		return Profile->GetCookedMaps();
+	}
+
+	float FGenericProfileMarshaller::GetMapListHeight() const
+	{
+		return MapListHeight;
+	}
+
+	void FGenericProfileMarshaller::SetMapListHeight( float NewHeight )
+	{
+		static const float MinMapListHeight = 100.0f;
+
+		MapListHeight = FMath::Max(NewHeight, MinMapListHeight);
+	}
+
+	FGenericProfileMarshaller::EMapOption FGenericProfileMarshaller::GetMapOption() const
+	{
+		return MapOption;
+	}
+
+	void FGenericProfileMarshaller::SetMapOption( EMapOption NewMapOption )
+	{
+		bool bShow = (NewMapOption == EMapOption::Selected);
+
+		MapOption = NewMapOption;
+
+		if (bShow)
+		{
+			// restore the cooked maps again, if any
+			if (CachedMapsToCook.Num() > 0 && Profile->GetCookedMaps().Num() == 0)
+			{
+				SetMapsToCook(CachedMapsToCook);
+				CachedMapsToCook.Reset();
+			}
+		}
+		else
+		{
+			// to set the 'cook startup maps only', its necessary to remove all the cooked maps - take a copy of the values to allow it to be restored
+			CachedMapsToCook = Profile->GetCookedMaps();
+			SetMapsToCook(TArray<FString>());
+		}
+
+		OnPropertyChanged();
+
+		MapListView->RefreshMapList();
+	}
+
+
+
+
+	BEGIN_SLATE_FUNCTION_BUILD_OPTIMIZATION
+	TSharedRef<SWidget> FGenericProfileMarshaller::CreateDeployDeviceWidget()
+	{
+		if (ProfileType == EProfileType::Basic)
+		{
+			return SNew(SCustomLaunchDeviceListView)
+				.OnSelectionChanged(this, &FGenericProfileMarshaller::SetDeployDeviceIDs)
+				.SelectedDevices(this, &FGenericProfileMarshaller::GetDeployDeviceIDs)
+				.AllPlatforms(true)
+				.SingleSelect(true)
+			;
+		}
+		else
+		{
+			return SNew(SVerticalBox)
+
+				// device picker options
+				+SVerticalBox::Slot()
+				.AutoHeight()
+				.Padding(0,2)
+				[
+					SNew(SHorizontalBox)
+
+					+SHorizontalBox::Slot()
+					.Padding(0, 0)
+					.AutoWidth()
+					.VAlign(VAlign_Center)
+					[
+						SNew(SSegmentedControl<EDeployDeviceOption>)
+						.Value( this, &FGenericProfileMarshaller::GetDeployDeviceOption)
+						.OnValueChanged(this, &FGenericProfileMarshaller::SetDeployDeviceOption)
+
+						+SSegmentedControl<EDeployDeviceOption>::Slot(EDeployDeviceOption::Default)
+						.HAlign(HAlign_Center)
+						.VAlign(VAlign_Center)
+						[
+							SNew(STextBlock)
+							.Text(LOCTEXT("DefaultDeviceLabel", "Default Device"))
+							.Font(FCoreStyle::Get().GetFontStyle("SmallFont"))
+						]
+
+						+SSegmentedControl<EDeployDeviceOption>::Slot(EDeployDeviceOption::Selected)
+						.HAlign(HAlign_Center)
+						.VAlign(VAlign_Center)
+						[
+							SNew(STextBlock)
+							.Text(LOCTEXT("SelectedDevicesLabel", "Selected Devices"))
+							.Font(FCoreStyle::Get().GetFontStyle("SmallFont"))
+						]
+					]
+				]
+
+				// device picker list
+				+SVerticalBox::Slot()
+				.FillHeight(1)
+				[
+					SNew(SVerticalResizeBox)
+					.Visibility_Lambda([this](){ return (GetDeployDeviceOption() == EDeployDeviceOption::Selected) ? EVisibility::Visible : EVisibility::Collapsed; } )
+					.HandleHeight(4)
+					.ContentHeight(this, &FGenericProfileMarshaller::GetDeployDeviceListHeight)
+					.ContentHeightChanged(this, &FGenericProfileMarshaller::SetDeployDeviceListHeight)
+					.HandleColor( FAppStyle::Get().GetSlateColor("Colors.Secondary").GetSpecifiedColor() )
+					[
+						SAssignNew(DeployDeviceListView, SCustomLaunchDeviceListView)
+						.OnDeviceRemoved( this, &FGenericProfileMarshaller::OnDeviceRemoved)
+						.OnSelectionChanged(this, &FGenericProfileMarshaller::SetDeployDeviceIDs)
+						.SelectedDevices(this, &FGenericProfileMarshaller::GetDeployDeviceIDs)
+						.Platforms(this, &FGenericProfileMarshaller::GetSelectedPlatforms)
+					]
+				]
+			;
+		}
+	}
+	END_SLATE_FUNCTION_BUILD_OPTIMIZATION
+
+
+	void FGenericProfileMarshaller::SetDeployDeviceIDs(TArray<FString> DeployDeviceIDs)
+	{
+		if (Profile->GetDeployedDeviceGroup() == nullptr)
+		{
+			Profile->SetDeployedDeviceGroup(Model->GetProfileManager()->AddNewDeviceGroup());
+		}
+
+		Profile->GetDeployedDeviceGroup()->RemoveAllDevices(); 
+		for ( const FString& DeviceID : DeployDeviceIDs )
+		{
+			Profile->GetDeployedDeviceGroup()->AddDevice(DeviceID);
+		}
+
+		if (ProfileType == EProfileType::Basic) // @fixme: should this be done for custom profiles too?
+		{
+			Model->UpdatedCookedPlatformsFromDeployDeviceProxy(Profile);
+		}
+
+		OnPropertyChanged();
+	}
+
+	TArray<FString> FGenericProfileMarshaller::GetDeployDeviceIDs() const
+	{
+		return Profile->GetDeployedDeviceGroup()->GetDeviceIDs();
+	}
+
+	float FGenericProfileMarshaller::GetDeployDeviceListHeight() const
+	{
+		return DeployDeviceListHeight;
+	}
+
+	void FGenericProfileMarshaller::SetDeployDeviceListHeight( float NewHeight )
+	{
+		static const float MinDeployDeviceListHeight = 100.0f;
+
+		DeployDeviceListHeight = FMath::Max(NewHeight, MinDeployDeviceListHeight);
+	}
+
+	void FGenericProfileMarshaller::OnDeviceRemoved(FString DeviceID)
+	{
+		CachedDeployDeviceIDs.Remove(DeviceID);
+	}
+
+	FGenericProfileMarshaller::EDeployDeviceOption FGenericProfileMarshaller::GetDeployDeviceOption() const
+	{
+		return DeployDeviceOption;
+	}
+
+	void FGenericProfileMarshaller::SetDeployDeviceOption( EDeployDeviceOption NewDeployDeviceOption )
+	{
+		bool bShow = (NewDeployDeviceOption == EDeployDeviceOption::Selected);
+
+		DeployDeviceOption = NewDeployDeviceOption;
+
+		if (bShow)
+		{
+			// restore the deployed device list again, if any
+			if (CachedDeployDeviceIDs.Num() > 0 && Profile->GetDeployedDeviceGroup()->GetDeviceIDs().Num() == 0)
+			{
+				SetDeployDeviceIDs(Profile->GetDeployedDeviceGroup()->GetDeviceIDs());
+				CachedDeployDeviceIDs.Reset();
+			}
+		}
+		else
+		{
+			// to set the 'default' deploy option, its necessary to remove all the devices - take a copy of the values to allow it to be restored
+			CachedDeployDeviceIDs = Profile->GetDeployedDeviceGroup()->GetDeviceIDs();
+			SetDeployDeviceIDs(TArray<FString>());
+		}
+
+		OnPropertyChanged();
+
+		DeployDeviceListView->RefreshDeviceList();
+	}
+
+
+	void FGenericProfileMarshaller::SetBuild( bool bBuild )
+	{
+		if (!bBuild)
+		{
+			Profile->SetBuildMode(ELauncherProfileBuildModes::DoNotBuild);
+		}
+		else if (GetForceBuild(Profile))
+		{
+			Profile->SetBuildMode(ELauncherProfileBuildModes::Build);
+		}
+		else
+		{
+			Profile->SetBuildMode(ELauncherProfileBuildModes::Auto);
+		}
+
+		OnPropertyChanged();
+	}
+
+	bool FGenericProfileMarshaller::GetBuild(ILauncherProfilePtr InProfile) const
+	{
+		if (InProfile == nullptr)
+		{
+			InProfile = Profile;
+		}
+		return (InProfile->GetBuildMode() != ELauncherProfileBuildModes::DoNotBuild);
+	}
+
+	void FGenericProfileMarshaller::SetForceBuild( bool bForceBuild )
+	{
+		if (!GetBuild())
+		{
+			Profile->SetBuildMode(ELauncherProfileBuildModes::DoNotBuild);
+		}
+		else if (bForceBuild)
+		{
+			Profile->SetBuildMode(ELauncherProfileBuildModes::Build);
+		}
+		else
+		{
+			Profile->SetBuildMode(ELauncherProfileBuildModes::Auto);
+		}
+
+		OnPropertyChanged();
+	}
+
+	bool FGenericProfileMarshaller::GetForceBuild(ILauncherProfilePtr InProfile) const
+	{
+		if (InProfile == nullptr)
+		{
+			InProfile = Profile;
+		}
+		return (InProfile->GetBuildMode() == ELauncherProfileBuildModes::Build);
+	}
+
+	void FGenericProfileMarshaller::SetArchitecture( FString Architecture )
+	{
+		// clear existing architectures
+		TArray<FString> Empty;
+		Profile->SetServerArchitectures(Empty);
+		Profile->SetEditorArchitectures(Empty);
+		Profile->SetClientArchitectures(Empty);
+
+		// set new single architecture
+		if (!Architecture.IsEmpty())
+		{
+			TArray<FString> Single;
+			Single.Add(Architecture);
+
+			switch (CachedBuildTargetType)
+			{
+				case EBuildTargetType::Server: Profile->SetServerArchitectures(Single); break;
+				case EBuildTargetType::Editor: Profile->SetEditorArchitectures(Single); break;
+				default:                       Profile->SetClientArchitectures(Single); break;
+			}
+		}
+
+		OnPropertyChanged();
+	}
+
+	FString FGenericProfileMarshaller::GetArchitecture() const
+	{
+		TArray<FString> Architectures;
+
+		switch (CachedBuildTargetType)
+		{
+			case EBuildTargetType::Server: Architectures = Profile->GetServerArchitectures(); break;
+			case EBuildTargetType::Editor: Architectures = Profile->GetEditorArchitectures(); break;
+			default:                       Architectures = Profile->GetClientArchitectures(); break;
+		}
+
+		return (Architectures.Num() > 0) ? Architectures[0] : FString();
+	}
+
+	FText FGenericProfileMarshaller::GetArchitectureDisplayName( FString Architecture )
+	{
+		if (Architecture.IsEmpty())
+		{
+			return LOCTEXT("DefaultArchName", "Project Default");
+		}
+		else if (Architecture == FPlatformMisc::GetHostArchitecture() && Model->IsHostPlatform(Profile))
+		{
+			return FText::Format(LOCTEXT("HostArchLabel", "{0} (this platform)"), FText::FromString(Architecture) );
+		}
+		else
+		{
+			return FText::FromString(Architecture);
+		}
+	}
+
+
+
+	void FGenericProfileMarshaller::SetDeployToDevice( bool bDeployToDevice )
+	{
+		if (ContentScheme != EContentScheme::CookOnTheFly)
+		{
+			Profile->SetDeploymentMode(bDeployToDevice ? ELauncherProfileDeploymentModes::CopyToDevice : ELauncherProfileDeploymentModes::DoNotDeploy );
+			OnPropertyChanged();
+		}
+	}
+
+	bool FGenericProfileMarshaller::GetDeployToDevice(ILauncherProfilePtr InProfile) const
+	{
+		if (InProfile == nullptr)
+		{
+			InProfile = Profile;
+		}
+		return InProfile->GetDeploymentMode() != ELauncherProfileDeploymentModes::DoNotDeploy;
+	}
+
+
+	void FGenericProfileMarshaller::SetIsRunning( bool bRun )
+	{
+		Profile->SetLaunchMode(bRun ? ELauncherProfileLaunchModes::DefaultRole : ELauncherProfileLaunchModes::DoNotLaunch);
+	}
+
+	bool FGenericProfileMarshaller::GetIsRunning( ILauncherProfilePtr InProfile) const
+	{
+		if (InProfile == nullptr)
+		{
+			InProfile = Profile;
+		}
+	
+		return InProfile->GetLaunchMode() != ELauncherProfileLaunchModes::DoNotLaunch;
+	}
+
+	void FGenericProfileMarshaller::OnInitialMapChanged( TSharedPtr<FString> InitialMap, ESelectInfo::Type )
+	{
+		if (InitialMap.IsValid())
+		{
+			Profile->GetDefaultLaunchRole()->SetInitialMap(*InitialMap);
+		}
+		else
+		{
+			Profile->GetDefaultLaunchRole()->SetInitialMap(FString());
+		}
+
+		OnPropertyChanged();
+	}
+
+	TSharedPtr<FString> FGenericProfileMarshaller::GetInitialMap() const
+	{
+		FString InitialMap = Profile->GetDefaultLaunchRole()->GetInitialMap();
+		const TSharedPtr<FString>* FoundMapPtr = CachedStartupMaps.FindByPredicate( [InitialMap]( const TSharedPtr<FString>& Map )
+		{
+			return Map.IsValid() && (InitialMap == *Map);
+		});
+
+		if (FoundMapPtr)
+		{
+			return *FoundMapPtr;
+		}
+		else
+		{
+			return MakeShared<FString>();
+		}
+	}
+
+}
+
+#undef LOCTEXT_NAMESPACE
diff --git a/Engine/Plugins/Developer/CustomLaunchUI/Source/Private/Marshalling/ILaunchProfileMarshaller.cpp b/Engine/Plugins/Developer/CustomLaunchUI/Source/Private/Marshalling/ILaunchProfileMarshaller.cpp
--- a/Engine/Plugins/Developer/CustomLaunchUI/Source/Private/Marshalling/ILaunchProfileMarshaller.cpp
+++ b/Engine/Plugins/Developer/CustomLaunchUI/Source/Private/Marshalling/ILaunchProfileMarshaller.cpp
@@ -1,82 +1,100 @@
 // Copyright Epic Games, Inc. All Rights Reserved.
 
 #include "Marshalling/ILaunchProfileMarshaller.h"
+#include "HAL/CriticalSection.h"
+#include "Misc/ScopeLock.h"
 
-
-static TArray<TSharedPtr<ILaunchProfileMarshallerFactory>> GMarshallerFactories;
+namespace CustomLaunchUI
+{
+	static TArray<TSharedPtr<ILaunchProfileMarshallerFactory>> GMarshallerFactories;
+	static FCriticalSection GMarshallerFactoriesCS;
 
 
 
-class FEmptyProfileMarshaller : public ILaunchProfileMarshaller
-{
-public:
-	FEmptyProfileMarshaller(const ILauncherProfilePtr& InProfile, TSharedRef<CustomLaunchUI::FModel> InModel)
-		: TreeData(MakeShared<FLaunchProfileTreeData>(InProfile, InModel, this))
+	class FEmptyProfileMarshaller : public ILaunchProfileMarshaller
 	{
-	}
-	virtual ~FEmptyProfileMarshaller() = default;
+	public:
+		FEmptyProfileMarshaller(const ILauncherProfilePtr& InProfile, TSharedRef<FModel> InModel)
+			: TreeData(MakeShared<FLaunchProfileTreeData>(InProfile, InModel, this))
+		{
+		}
+		virtual ~FEmptyProfileMarshaller() = default;
 
-	virtual void Construct() override
-	{
-	}
+		virtual void Construct() override
+		{
+		}
 
-	virtual FString GetName() const override
-	{
-		return TEXT("Empty");
-	}
+		virtual FString GetName() const override
+		{
+			return TEXT("Empty");
+		}
 
-	virtual FLaunchProfileTreeDataRef GetProfileTree() override
-	{
-		return TreeData;
-	}
+		virtual FLaunchProfileTreeDataRef GetProfileTree() override
+		{
+			return TreeData;
+		}
 
-	virtual void OnPropertyChanged() override
-	{
-	}
+		virtual void OnPropertyChanged() override
+		{
+		}
 
-protected:
-	FLaunchProfileTreeDataRef TreeData;
-};
+	protected:
+		FLaunchProfileTreeDataRef TreeData;
+	};
 
 
 
 
 
-void ILaunchProfileMarshallerFactory::Register( TSharedRef<ILaunchProfileMarshallerFactory> Marshaller )
-{
-	GMarshallerFactories.Add(Marshaller);
-	GMarshallerFactories.Sort( []( TSharedPtr<ILaunchProfileMarshallerFactory> A, TSharedPtr<ILaunchProfileMarshallerFactory> B)
+	void RegisterMarshallerFactory( TSharedRef<ILaunchProfileMarshallerFactory> MarshallerFactory )
 	{
-		return A->GetPriority() > B->GetPriority();
-	});
+		FScopeLock Lock(&GMarshallerFactoriesCS);
 
-}
+		GMarshallerFactories.Add(MarshallerFactory);
+		GMarshallerFactories.Sort( []( TSharedPtr<ILaunchProfileMarshallerFactory> A, TSharedPtr<ILaunchProfileMarshallerFactory> B)
+		{
+			return A->GetPriority() > B->GetPriority();
+		});
 
-void ILaunchProfileMarshallerFactory::Unregister( TSharedRef<ILaunchProfileMarshallerFactory> Marshaller )
-{
-	GMarshallerFactories.Remove(Marshaller);
-}
+	}
 
+	void UnregisterMarshallerFactory( TSharedRef<ILaunchProfileMarshallerFactory> MarshallerFactory )
+	{
+		FScopeLock Lock(&GMarshallerFactoriesCS);
 
-TSharedRef<ILaunchProfileMarshaller> ILaunchProfileMarshallerFactory::Marshall( const ILauncherProfilePtr& InProfile, TSharedRef<CustomLaunchUI::FModel> InModel )
-{
-	if (InProfile.IsValid())
+		GMarshallerFactories.Remove(MarshallerFactory);
+	}
+
+
+	TSharedRef<ILaunchProfileMarshaller> CreateMarshaller( const ILauncherProfilePtr& InProfile, TSharedRef<FModel> InModel )
 	{
-		// find a marshaller that can handle this profile
-		for (const TSharedPtr<ILaunchProfileMarshallerFactory>& Factory : GMarshallerFactories)
+		if (InProfile.IsValid())
 		{
-			TSharedPtr<ILaunchProfileMarshaller> Marshaller = Factory->TryCreateMarshaller(InProfile.ToSharedRef(), InModel);
-			if (Marshaller.IsValid())
+			EProfileType ProfileType = InModel->GetProfileType(InProfile.ToSharedRef());
+
+			FScopeLock Lock(&GMarshallerFactoriesCS);
+
+			// find a marshaller that can handle this profile
+			for (const TSharedPtr<ILaunchProfileMarshallerFactory>& Factory : GMarshallerFactories)
 			{
-				Marshaller->Construct();
-				return Marshaller.ToSharedRef();
+				if (!Factory->IsProfileTypeSupported(ProfileType))
+				{
+					continue;
+				}
+
+				TSharedPtr<ILaunchProfileMarshaller> Marshaller = Factory->TryCreateMarshaller(InProfile.ToSharedRef(), InModel);
+				if (Marshaller.IsValid())
+				{
+					Marshaller->Construct();
+					return Marshaller.ToSharedRef();
+				}
 			}
 		}
-	}
 
-	// cannot marshall
-	TSharedRef<ILaunchProfileMarshaller> EmptyMarshaller = MakeShared<FEmptyProfileMarshaller>(InProfile, InModel);
-	EmptyMarshaller->Construct();
-	return EmptyMarshaller;
+		// cannot marshall
+		TSharedRef<ILaunchProfileMarshaller> EmptyMarshaller = MakeShared<FEmptyProfileMarshaller>(InProfile, InModel);
+		EmptyMarshaller->Construct();
+		return EmptyMarshaller;
+	}
 }
 
diff --git a/Engine/Plugins/Developer/CustomLaunchUI/Source/Private/Marshalling/ILaunchProfileMarshaller.h b/Engine/Plugins/Developer/CustomLaunchUI/Source/Private/Marshalling/ILaunchProfileMarshaller.h
deleted file mode 100644
diff --git a/Engine/Plugins/Developer/CustomLaunchUI/Source/Private/Marshalling/LaunchProfileTreeData.cpp b/Engine/Plugins/Developer/CustomLaunchUI/Source/Private/Marshalling/LaunchProfileTreeData.cpp
--- a/Engine/Plugins/Developer/CustomLaunchUI/Source/Private/Marshalling/LaunchProfileTreeData.cpp
+++ b/Engine/Plugins/Developer/CustomLaunchUI/Source/Private/Marshalling/LaunchProfileTreeData.cpp
@@ -16,239 +16,240 @@
 
 #define LOCTEXT_NAMESPACE "SCustomLaunchCustomProfileEditor"
 
-
-FLaunchProfileTreeData::FLaunchProfileTreeData(ILauncherProfilePtr InProfile, TSharedRef<CustomLaunchUI::FModel> InModel, ILaunchProfileMarshaller* InMarshaller) 
-	: Profile(InProfile)
-	, Model(InModel)
-	, Marshaller(InMarshaller)
+namespace CustomLaunchUI
 {
-	check(Marshaller);
-}
-
-FLaunchProfileTreeNode& FLaunchProfileTreeData::AddHeading( FText InName )
-{
-	FLaunchProfileTreeNodePtr TreeNode = MakeShared<FLaunchProfileTreeNode>(this);
-	TreeNode->Name = InName;
-	Nodes.Add(TreeNode);
+	FLaunchProfileTreeData::FLaunchProfileTreeData(ILauncherProfilePtr InProfile, TSharedRef<CustomLaunchUI::FModel> InModel, ILaunchProfileMarshaller* InMarshaller) 
+		: Profile(InProfile)
+		, Model(InModel)
+		, Marshaller(InMarshaller)
+	{
+		check(Marshaller);
+	}
 
-	return *TreeNode.Get();
-}
+	FLaunchProfileTreeNode& FLaunchProfileTreeData::AddHeading( FText InName )
+	{
+		FLaunchProfileTreeNodePtr TreeNode = MakeShared<FLaunchProfileTreeNode>(this);
+		TreeNode->Name = InName;
+		Nodes.Add(TreeNode);
 
+		return *TreeNode.Get();
+	}
 
 
-FLaunchProfileTreeNode::FLaunchProfileTreeNode( const FLaunchProfileTreeData* InTreeData ) 
-	: TreeData(InTreeData)
-{
-}
 
-FLaunchProfileTreeNode& FLaunchProfileTreeNode::AddWidget( FText InName, FCallbacks&& InWidgetCallbacks, TSharedRef<SWidget> InValueWidget )
-{
-	FLaunchProfileTreeNodePtr TreeNode = MakeShared<FLaunchProfileTreeNode>(TreeData);
-	TreeNode->Name = InName;
-	TreeNode->Widget = InValueWidget;
-	TreeNode->Callbacks = InWidgetCallbacks;
-	Children.Add(TreeNode);
+	FLaunchProfileTreeNode::FLaunchProfileTreeNode( const FLaunchProfileTreeData* InTreeData ) 
+		: TreeData(InTreeData)
+	{
+	}
 
-	return *this;
-}
+	FLaunchProfileTreeNode& FLaunchProfileTreeNode::AddWidget( FText InName, FCallbacks&& InWidgetCallbacks, TSharedRef<SWidget> InValueWidget )
+	{
+		FLaunchProfileTreeNodePtr TreeNode = MakeShared<FLaunchProfileTreeNode>(TreeData);
+		TreeNode->Name = InName;
+		TreeNode->Widget = InValueWidget;
+		TreeNode->Callbacks = InWidgetCallbacks;
+		Children.Add(TreeNode);
 
-FLaunchProfileTreeNode& FLaunchProfileTreeNode::AddWidget( FText InName, TSharedRef<SWidget> InValueWidget )
-{
-	return AddWidget( InName, {}, InValueWidget );
-}
+		return *this;
+	}
 
+	FLaunchProfileTreeNode& FLaunchProfileTreeNode::AddWidget( FText InName, TSharedRef<SWidget> InValueWidget )
+	{
+		return AddWidget( InName, {}, InValueWidget );
+	}
 
-FLaunchProfileTreeNode& FLaunchProfileTreeNode::AddBoolean( FText InName, FBooleanCallbacks&& BooleanCallbacks )
-{
-	check(BooleanCallbacks.GetValue);
-	check(BooleanCallbacks.SetValue);
 
-	FCallbacks WidgetCallbacks
-	{
-		.IsVisible = BooleanCallbacks.IsVisible,
-		.IsEnabled = BooleanCallbacks.IsEnabled,
-	};
-	if (BooleanCallbacks.GetDefaultValue != nullptr)
+	FLaunchProfileTreeNode& FLaunchProfileTreeNode::AddBoolean( FText InName, FBooleanCallbacks&& BooleanCallbacks )
 	{
-		WidgetCallbacks.IsDefault = [BooleanCallbacks]()
-		{
-			return BooleanCallbacks.GetValue() == BooleanCallbacks.GetDefaultValue();
-		};
+		check(BooleanCallbacks.GetValue);
+		check(BooleanCallbacks.SetValue);
 
-		WidgetCallbacks.SetToDefault = [BooleanCallbacks]()
+		FCallbacks WidgetCallbacks
 		{
-			BooleanCallbacks.SetValue(BooleanCallbacks.GetDefaultValue());
-			// nb. OnPropertyChanged is called via SCustomLaunchWidgetTableRow's set to default
+			.IsVisible = BooleanCallbacks.IsVisible,
+			.IsEnabled = BooleanCallbacks.IsEnabled,
 		};
-	};
-
-	auto SetCheckState = [this, BooleanCallbacks]( ECheckBoxState CheckState )
-	{
-		BooleanCallbacks.SetValue( (CheckState == ECheckBoxState::Checked) );
-		TreeData->Marshaller->OnPropertyChanged();
-	};
-
-	auto GetCheckState = [BooleanCallbacks]()
-	{
-		return BooleanCallbacks.GetValue() ? ECheckBoxState::Checked : ECheckBoxState::Unchecked;
-	};
-
-	return AddWidget( InName,
-		MoveTemp(WidgetCallbacks),
-		SNew(SCheckBox)
-		.OnCheckStateChanged_Lambda(SetCheckState)
-		.IsChecked_Lambda(GetCheckState)
-	);
-}
+		if (BooleanCallbacks.GetDefaultValue != nullptr)
+		{
+			WidgetCallbacks.IsDefault = [BooleanCallbacks]()
+			{
+				return BooleanCallbacks.GetValue() == BooleanCallbacks.GetDefaultValue();
+			};
 
-FLaunchProfileTreeNode& FLaunchProfileTreeNode::AddString( FText InName, FStringCallbacks&& StringCallbacks )
-{
-	check(StringCallbacks.GetValue);
-	check(StringCallbacks.SetValue);
+			WidgetCallbacks.SetToDefault = [BooleanCallbacks]()
+			{
+				BooleanCallbacks.SetValue(BooleanCallbacks.GetDefaultValue());
+				// nb. OnPropertyChanged is called via SCustomLaunchWidgetTableRow's set to default
+			};
+		};
 
-	FCallbacks WidgetCallbacks
-	{
-		.IsVisible = StringCallbacks.IsVisible,
-		.IsEnabled = StringCallbacks.IsEnabled,
-	};
-	if (StringCallbacks.GetDefaultValue != nullptr)
-	{
-		WidgetCallbacks.IsDefault = [StringCallbacks]()
+		auto SetCheckState = [this, BooleanCallbacks]( ECheckBoxState CheckState )
 		{
-			return StringCallbacks.GetValue() == StringCallbacks.GetDefaultValue();
+			BooleanCallbacks.SetValue( (CheckState == ECheckBoxState::Checked) );
+			TreeData->Marshaller->OnPropertyChanged();
 		};
 
-		WidgetCallbacks.SetToDefault = [StringCallbacks]()
+		auto GetCheckState = [BooleanCallbacks]()
 		{
-			StringCallbacks.SetValue(StringCallbacks.GetDefaultValue());
-			// nb. OnPropertyChanged is called via SCustomLaunchWidgetTableRow's set to default
+			return BooleanCallbacks.GetValue() ? ECheckBoxState::Checked : ECheckBoxState::Unchecked;
 		};
-	};
 
-	auto SetString = [this, StringCallbacks]( const FText& InText, ETextCommit::Type InTextCommit )
-	{
-		StringCallbacks.SetValue( InText.ToString() );
-		TreeData->Marshaller->OnPropertyChanged();
-	};
+		return AddWidget( InName,
+			MoveTemp(WidgetCallbacks),
+			SNew(SCheckBox)
+			.OnCheckStateChanged_Lambda(SetCheckState)
+			.IsChecked_Lambda(GetCheckState)
+		);
+	}
 
-	auto GetString = [StringCallbacks]()
+	FLaunchProfileTreeNode& FLaunchProfileTreeNode::AddString( FText InName, FStringCallbacks&& StringCallbacks )
 	{
-		return FText::FromString(StringCallbacks.GetValue());
-	};
-
-	return AddWidget( InName,
-		MoveTemp(WidgetCallbacks),
-		SNew(SMultiLineEditableTextBox)
-		.AllowMultiLine(false)
-		.AutoWrapText(true)
-		.Font(FCoreStyle::Get().GetFontStyle("SmallFont"))
-		.Text_Lambda( GetString )
-		.OnTextCommitted_Lambda( SetString )
-	);
+		check(StringCallbacks.GetValue);
+		check(StringCallbacks.SetValue);
 
-}
+		FCallbacks WidgetCallbacks
+		{
+			.IsVisible = StringCallbacks.IsVisible,
+			.IsEnabled = StringCallbacks.IsEnabled,
+		};
+		if (StringCallbacks.GetDefaultValue != nullptr)
+		{
+			WidgetCallbacks.IsDefault = [StringCallbacks]()
+			{
+				return StringCallbacks.GetValue() == StringCallbacks.GetDefaultValue();
+			};
 
-FLaunchProfileTreeNode& FLaunchProfileTreeNode::AddDirectoryString( FText InName, FStringCallbacks&& StringCallbacks )
-{
-	check(StringCallbacks.GetValue);
-	check(StringCallbacks.SetValue);
+			WidgetCallbacks.SetToDefault = [StringCallbacks]()
+			{
+				StringCallbacks.SetValue(StringCallbacks.GetDefaultValue());
+				// nb. OnPropertyChanged is called via SCustomLaunchWidgetTableRow's set to default
+			};
+		};
 
-	FCallbacks WidgetCallbacks
-	{
-		.IsVisible = StringCallbacks.IsVisible,
-		.IsEnabled = StringCallbacks.IsEnabled,
-	};
-	if (StringCallbacks.GetDefaultValue != nullptr)
-	{
-		WidgetCallbacks.IsDefault = [StringCallbacks]()
+		auto SetString = [this, StringCallbacks]( const FText& InText, ETextCommit::Type InTextCommit )
 		{
-			return StringCallbacks.GetValue() == StringCallbacks.GetDefaultValue();
+			StringCallbacks.SetValue( InText.ToString() );
+			TreeData->Marshaller->OnPropertyChanged();
 		};
 
-		WidgetCallbacks.SetToDefault = [StringCallbacks]()
+		auto GetString = [StringCallbacks]()
 		{
-			StringCallbacks.SetValue(StringCallbacks.GetDefaultValue());
-			// nb. OnPropertyChanged is called via SCustomLaunchWidgetTableRow's set to default
+			return FText::FromString(StringCallbacks.GetValue());
 		};
-	};
 
-	auto SetString = [this, StringCallbacks]( const FText& InText, ETextCommit::Type InTextCommit )
-	{
-		StringCallbacks.SetValue( InText.ToString() );
-		TreeData->Marshaller->OnPropertyChanged();
-	};
+		return AddWidget( InName,
+			MoveTemp(WidgetCallbacks),
+			SNew(SMultiLineEditableTextBox)
+			.AllowMultiLine(false)
+			.AutoWrapText(true)
+			.Font(FCoreStyle::Get().GetFontStyle("SmallFont"))
+			.Text_Lambda( GetString )
+			.OnTextCommitted_Lambda( SetString )
+		);
 
-	auto GetString = [StringCallbacks]()
-	{
-		return FText::FromString(StringCallbacks.GetValue());
-	};
+	}
 
-	auto OnBrowse = [this,StringCallbacks]()
+	FLaunchProfileTreeNode& FLaunchProfileTreeNode::AddDirectoryString( FText InName, FStringCallbacks&& StringCallbacks )
 	{
-		FString InitialDirectory = StringCallbacks.GetValue();
-		if (InitialDirectory.IsEmpty())
+		check(StringCallbacks.GetValue);
+		check(StringCallbacks.SetValue);
+
+		FCallbacks WidgetCallbacks
+		{
+			.IsVisible = StringCallbacks.IsVisible,
+			.IsEnabled = StringCallbacks.IsEnabled,
+		};
+		if (StringCallbacks.GetDefaultValue != nullptr)
+		{
+			WidgetCallbacks.IsDefault = [StringCallbacks]()
+			{
+				return StringCallbacks.GetValue() == StringCallbacks.GetDefaultValue();
+			};
+
+			WidgetCallbacks.SetToDefault = [StringCallbacks]()
+			{
+				StringCallbacks.SetValue(StringCallbacks.GetDefaultValue());
+				// nb. OnPropertyChanged is called via SCustomLaunchWidgetTableRow's set to default
+			};
+		};
+
+		auto SetString = [this, StringCallbacks]( const FText& InText, ETextCommit::Type InTextCommit )
 		{
-			InitialDirectory = this->TreeData->Profile->GetProjectBasePath();
-		}
-		if (!InitialDirectory.IsEmpty() && FPaths::IsRelative(InitialDirectory))
+			StringCallbacks.SetValue( InText.ToString() );
+			TreeData->Marshaller->OnPropertyChanged();
+		};
+
+		auto GetString = [StringCallbacks]()
 		{
-			InitialDirectory = FPaths::Combine( FPaths::RootDir(), InitialDirectory );
-		}
+			return FText::FromString(StringCallbacks.GetValue());
+		};
 
-		const void* ParentWindowHandle = FSlateApplication::Get().FindBestParentWindowHandleForDialogs(nullptr);
-		FString OutDirectory;
-		if (FDesktopPlatformModule::Get()->OpenDirectoryDialog(ParentWindowHandle, LOCTEXT("FolderDialogTitle", "Choose a directory").ToString(), InitialDirectory, OutDirectory))
+		auto OnBrowse = [this,StringCallbacks]()
 		{
-			if (FPaths::IsUnderDirectory(OutDirectory, FPaths::RootDir()))
+			FString InitialDirectory = StringCallbacks.GetValue();
+			if (InitialDirectory.IsEmpty())
 			{
-				FPaths::MakePathRelativeTo(OutDirectory, *FPaths::RootDir());
+				InitialDirectory = this->TreeData->Profile->GetProjectBasePath();
+			}
+			if (!InitialDirectory.IsEmpty() && FPaths::IsRelative(InitialDirectory))
+			{
+				InitialDirectory = FPaths::Combine( FPaths::RootDir(), InitialDirectory );
 			}
 
-			StringCallbacks.SetValue( OutDirectory );
-			TreeData->Marshaller->OnPropertyChanged();
-		}
+			const void* ParentWindowHandle = FSlateApplication::Get().FindBestParentWindowHandleForDialogs(nullptr);
+			FString OutDirectory;
+			if (FDesktopPlatformModule::Get()->OpenDirectoryDialog(ParentWindowHandle, LOCTEXT("FolderDialogTitle", "Choose a directory").ToString(), InitialDirectory, OutDirectory))
+			{
+				if (FPaths::IsUnderDirectory(OutDirectory, FPaths::RootDir()))
+				{
+					FPaths::MakePathRelativeTo(OutDirectory, *FPaths::RootDir());
+				}
 
-		return FReply::Handled();
-	};
+				StringCallbacks.SetValue( OutDirectory );
+				TreeData->Marshaller->OnPropertyChanged();
+			}
 
+			return FReply::Handled();
+		};
 
-	return AddWidget( InName,
-		MoveTemp(WidgetCallbacks),
-		SNew(SHorizontalBox)
 
-		// path field
-		+SHorizontalBox::Slot()
-		.FillWidth(1)
-		.VAlign(VAlign_Center)
-		[
-			SNew(SEditableTextBox)
-			.Font(FCoreStyle::Get().GetFontStyle("SmallFont"))
-			.Text_Lambda( GetString )
-			.OnTextCommitted_Lambda( SetString )
-		]
-
-		// browse button
-		+SHorizontalBox::Slot()
-		.AutoWidth()
-		.Padding(FMargin(4, 0, 0, 0))
-		.VAlign(VAlign_Center)
-		[
-			SNew(SButton)
-			.ButtonStyle(FCustomLaunchUIStyle::Get(), "HoverHintOnly")
-			.ToolTipText(LOCTEXT("DirBrowseTip", "Browse for a folder"))
-			.OnClicked_Lambda(OnBrowse)
-			.ContentPadding(2)
+		return AddWidget( InName,
+			MoveTemp(WidgetCallbacks),
+			SNew(SHorizontalBox)
+
+			// path field
+			+SHorizontalBox::Slot()
+			.FillWidth(1)
+			.VAlign(VAlign_Center)
 			[
-				SNew(SImage)
-				.Image(FCustomLaunchUIStyle::Get().GetBrush("PathPickerButton"))
-				.ColorAndOpacity(FSlateColor::UseForeground())
+				SNew(SEditableTextBox)
+				.Font(FCoreStyle::Get().GetFontStyle("SmallFont"))
+				.Text_Lambda( GetString )
+				.OnTextCommitted_Lambda( SetString )
 			]
-		]
-	);
-}
 
+			// browse button
+			+SHorizontalBox::Slot()
+			.AutoWidth()
+			.Padding(FMargin(4, 0, 0, 0))
+			.VAlign(VAlign_Center)
+			[
+				SNew(SButton)
+				.ButtonStyle(FCustomLaunchUIStyle::Get(), "HoverHintOnly")
+				.ToolTipText(LOCTEXT("DirBrowseTip", "Browse for a folder"))
+				.OnClicked_Lambda(OnBrowse)
+				.ContentPadding(2)
+				[
+					SNew(SImage)
+					.Image(FCustomLaunchUIStyle::Get().GetBrush("PathPickerButton"))
+					.ColorAndOpacity(FSlateColor::UseForeground())
+				]
+			]
+		);
+	}
 
 
 
+}
 
 #undef LOCTEXT_NAMESPACE
diff --git a/Engine/Plugins/Developer/CustomLaunchUI/Source/Private/Marshalling/LaunchProfileTreeData.h b/Engine/Plugins/Developer/CustomLaunchUI/Source/Private/Marshalling/LaunchProfileTreeData.h
deleted file mode 100644
diff --git a/Engine/Plugins/Developer/CustomLaunchUI/Source/Private/Model/CustomLaunchUIModel.cpp b/Engine/Plugins/Developer/CustomLaunchUI/Source/Private/Model/CustomLaunchUIModel.cpp
--- a/Engine/Plugins/Developer/CustomLaunchUI/Source/Private/Model/CustomLaunchUIModel.cpp
+++ b/Engine/Plugins/Developer/CustomLaunchUI/Source/Private/Model/CustomLaunchUIModel.cpp
@@ -84,6 +84,22 @@ namespace CustomLaunchUI
 	}
 
 
+	EProfileType FModel::GetProfileType(const ILauncherProfileRef& Profile) const
+	{
+		if (IsAdvancedProfile(Profile))
+		{
+			return EProfileType::Advanced;
+		}
+		else if (Profile == BasicLaunchProfile)
+		{
+			return EProfileType::Basic;
+		}
+		else
+		{
+			return EProfileType::Custom;
+		}
+	}
+
 	bool FModel::IsAdvancedProfile(const ILauncherProfileRef& Profile) const
 	{
 		return Profile->GetCookedPlatforms().Num() > 1 ||
diff --git a/Engine/Plugins/Developer/CustomLaunchUI/Source/Private/Model/CustomLaunchUIModel.h b/Engine/Plugins/Developer/CustomLaunchUI/Source/Private/Model/CustomLaunchUIModel.h
deleted file mode 100644
diff --git a/Engine/Plugins/Developer/CustomLaunchUI/Source/Private/Widgets/Profiles/SCustomLaunchCustomProfileEditor.cpp b/Engine/Plugins/Developer/CustomLaunchUI/Source/Private/Widgets/Profiles/SCustomLaunchCustomProfileEditor.cpp
--- a/Engine/Plugins/Developer/CustomLaunchUI/Source/Private/Widgets/Profiles/SCustomLaunchCustomProfileEditor.cpp
+++ b/Engine/Plugins/Developer/CustomLaunchUI/Source/Private/Widgets/Profiles/SCustomLaunchCustomProfileEditor.cpp
@@ -36,7 +36,10 @@
 
 #define LOCTEXT_NAMESPACE "SCustomLaunchCustomProfileEditor"
 
-
+namespace CustomLaunchUI
+{
+	extern TSharedRef<ILaunchProfileMarshaller> CreateMarshaller( const ILauncherProfilePtr& InProfile, TSharedRef<CustomLaunchUI::FModel> InModel );
+};
 
 
 
@@ -46,11 +49,11 @@ public:
 	SLATE_BEGIN_ARGS(SLaunchProfileCategoryTreeRow){}
 	SLATE_END_ARGS()
 
-	void Construct(const FArguments& InArgs, FLaunchProfileTreeNodePtr TreeNode, const TSharedRef<STableViewBase>& OwnerTable )
+	void Construct(const FArguments& InArgs, CustomLaunchUI::FLaunchProfileTreeNodePtr TreeNode, const TSharedRef<STableViewBase>& OwnerTable )
 	{
 		auto HasVisibleChildren = [TreeNode]()
 		{
-			for (FLaunchProfileTreeNodePtr ChildNode : TreeNode->Children)
+			for (CustomLaunchUI::FLaunchProfileTreeNodePtr ChildNode : TreeNode->Children)
 			{
 				if (ChildNode->Callbacks.IsVisible == nullptr || ChildNode->Callbacks.IsVisible())
 				{
@@ -122,7 +125,7 @@ public:
 	SLATE_END_ARGS()
 
 	BEGIN_SLATE_FUNCTION_BUILD_OPTIMIZATION
-	void Construct(const FArguments& InArgs, FLaunchProfileTreeNodePtr TreeNode, const TSharedRef<STableViewBase>& OwnerTable )
+	void Construct(const FArguments& InArgs, CustomLaunchUI::FLaunchProfileTreeNodePtr TreeNode, const TSharedRef<STableViewBase>& OwnerTable )
 	{
 		SplitterValue = InArgs._SplitterValue;
 		OnSplitterResized = InArgs._OnSplitterResized;
@@ -271,7 +274,7 @@ void SCustomLaunchCustomProfileEditor::Construct(const FArguments& InArgs, const
 {
 	Model = InModel;
 
-	Marshaller = ILaunchProfileMarshallerFactory::Marshall( nullptr, InModel );
+	Marshaller = CustomLaunchUI::CreateMarshaller( nullptr, InModel );
 
 
 	ChildSlot
@@ -287,7 +290,7 @@ void SCustomLaunchCustomProfileEditor::Construct(const FArguments& InArgs, const
 }
 END_SLATE_FUNCTION_BUILD_OPTIMIZATION
 
-TSharedRef<ITableRow> SCustomLaunchCustomProfileEditor::OnGenerateWidgetForTreeNode( FLaunchProfileTreeNodePtr TreeNode, const TSharedRef<STableViewBase>& OwnerTable )
+TSharedRef<ITableRow> SCustomLaunchCustomProfileEditor::OnGenerateWidgetForTreeNode( CustomLaunchUI::FLaunchProfileTreeNodePtr TreeNode, const TSharedRef<STableViewBase>& OwnerTable )
 {
 	if (TreeNode->Widget.IsValid())
 	{
@@ -314,7 +317,7 @@ TSharedRef<ITableRow> SCustomLaunchCustomProfileEditor::OnGenerateWidgetForTreeN
 	}
 }
 
-void SCustomLaunchCustomProfileEditor::OnGetChildren(FLaunchProfileTreeNodePtr TreeNode, TArray<FLaunchProfileTreeNodePtr>& OutChildren)
+void SCustomLaunchCustomProfileEditor::OnGetChildren(CustomLaunchUI::FLaunchProfileTreeNodePtr TreeNode, TArray<CustomLaunchUI::FLaunchProfileTreeNodePtr>& OutChildren)
 {
 	OutChildren = TreeNode->Children;
 }
@@ -350,10 +353,10 @@ void SCustomLaunchCustomProfileEditor::SetProfile( const ILauncherProfilePtr& Pr
 	CurrentProfile = Profile;
 
 
-	Marshaller = ILaunchProfileMarshallerFactory::Marshall( Profile, Model.ToSharedRef() );
+	Marshaller = CustomLaunchUI::CreateMarshaller( Profile, Model.ToSharedRef() );
 
 	TreeView->SetTreeItemsSource( &Marshaller->GetProfileTree()->Nodes );
-	for ( FLaunchProfileTreeNodePtr Node : TreeView->GetRootItems())
+	for ( CustomLaunchUI::FLaunchProfileTreeNodePtr Node : TreeView->GetRootItems())
 	{
 		TreeView->SetItemExpansion(Node, true);
 	}
diff --git a/Engine/Plugins/Developer/CustomLaunchUI/Source/Private/Widgets/Profiles/SCustomLaunchCustomProfileEditor.h b/Engine/Plugins/Developer/CustomLaunchUI/Source/Private/Widgets/Profiles/SCustomLaunchCustomProfileEditor.h
--- a/Engine/Plugins/Developer/CustomLaunchUI/Source/Private/Widgets/Profiles/SCustomLaunchCustomProfileEditor.h
+++ b/Engine/Plugins/Developer/CustomLaunchUI/Source/Private/Widgets/Profiles/SCustomLaunchCustomProfileEditor.h
@@ -16,8 +16,8 @@ class ITableRow;
 class STableViewBase;
 template<typename T> class STreeView;
 
-typedef STreeView<FLaunchProfileTreeNodePtr> SLaunchProfileTreeView;
-typedef STableRow<FLaunchProfileTreeNodePtr> SLaunchProfileTreeRow;
+typedef STreeView<CustomLaunchUI::FLaunchProfileTreeNodePtr> SLaunchProfileTreeView;
+typedef STableRow<CustomLaunchUI::FLaunchProfileTreeNodePtr> SLaunchProfileTreeRow;
 
 
 class SCustomLaunchCustomProfileEditor : public SCompoundWidget, public IScrollableWidget
@@ -40,13 +40,13 @@ public:
 
 private:
 
-	TSharedRef<ITableRow> OnGenerateWidgetForTreeNode( FLaunchProfileTreeNodePtr TreeNode, const TSharedRef<STableViewBase>& OwnerTable );
-	void OnGetChildren(FLaunchProfileTreeNodePtr TreeNode, TArray<FLaunchProfileTreeNodePtr>& OutChildren);
+	TSharedRef<ITableRow> OnGenerateWidgetForTreeNode( CustomLaunchUI::FLaunchProfileTreeNodePtr TreeNode, const TSharedRef<STableViewBase>& OwnerTable );
+	void OnGetChildren(CustomLaunchUI::FLaunchProfileTreeNodePtr TreeNode, TArray<CustomLaunchUI::FLaunchProfileTreeNodePtr>& OutChildren);
 
 	ILauncherProfilePtr CurrentProfile;
 	TSharedPtr<CustomLaunchUI::FModel> Model;
+	TSharedPtr<CustomLaunchUI::ILaunchProfileMarshaller> Marshaller;
 	TSharedPtr<SLaunchProfileTreeView> TreeView;
-	TSharedPtr<ILaunchProfileMarshaller> Marshaller;
 
 	float SplitterPos = 0.6f;
 	
diff --git a/Engine/Plugins/Developer/CustomLaunchUI/Source/Private/Widgets/Shared/SCustomLaunchBuildTargetCombo.h b/Engine/Plugins/Developer/CustomLaunchUI/Source/Private/Widgets/Shared/SCustomLaunchBuildTargetCombo.h
deleted file mode 100644
diff --git a/Engine/Plugins/Developer/CustomLaunchUI/Source/Private/Widgets/Shared/SCustomLaunchCombo.h b/Engine/Plugins/Developer/CustomLaunchUI/Source/Private/Widgets/Shared/SCustomLaunchCombo.h
deleted file mode 100644
diff --git a/Engine/Plugins/Developer/CustomLaunchUI/Source/Private/Widgets/Shared/SCustomLaunchContentSchemeCombo.h b/Engine/Plugins/Developer/CustomLaunchUI/Source/Private/Widgets/Shared/SCustomLaunchContentSchemeCombo.h
deleted file mode 100644
diff --git a/Engine/Plugins/Developer/CustomLaunchUI/Source/Private/Widgets/Shared/SCustomLaunchDeviceCombo.h b/Engine/Plugins/Developer/CustomLaunchUI/Source/Private/Widgets/Shared/SCustomLaunchDeviceCombo.h
deleted file mode 100644
diff --git a/Engine/Plugins/Developer/CustomLaunchUI/Source/Private/Widgets/Shared/SCustomLaunchDeviceListView.h b/Engine/Plugins/Developer/CustomLaunchUI/Source/Private/Widgets/Shared/SCustomLaunchDeviceListView.h
deleted file mode 100644
diff --git a/Engine/Plugins/Developer/CustomLaunchUI/Source/Private/Widgets/Shared/SCustomLaunchDeviceWidgetBase.h b/Engine/Plugins/Developer/CustomLaunchUI/Source/Private/Widgets/Shared/SCustomLaunchDeviceWidgetBase.h
deleted file mode 100644
diff --git a/Engine/Plugins/Developer/CustomLaunchUI/Source/Private/Widgets/Shared/SCustomLaunchMapListView.h b/Engine/Plugins/Developer/CustomLaunchUI/Source/Private/Widgets/Shared/SCustomLaunchMapListView.h
deleted file mode 100644
diff --git a/Engine/Plugins/Developer/CustomLaunchUI/Source/Private/Widgets/Shared/SCustomLaunchPlatformCombo.h b/Engine/Plugins/Developer/CustomLaunchUI/Source/Private/Widgets/Shared/SCustomLaunchPlatformCombo.h
deleted file mode 100644
diff --git a/Engine/Plugins/Developer/CustomLaunchUI/Source/Private/Widgets/Shared/SCustomLaunchProjectCombo.h b/Engine/Plugins/Developer/CustomLaunchUI/Source/Private/Widgets/Shared/SCustomLaunchProjectCombo.h
deleted file mode 100644
diff --git a/Engine/Plugins/Developer/CustomLaunchUI/Source/Public/CustomLaunchUIModule.h b/Engine/Plugins/Developer/CustomLaunchUI/Source/Public/CustomLaunchUIModule.h
new file mode 100644
--- /dev/null
+++ b/Engine/Plugins/Developer/CustomLaunchUI/Source/Public/CustomLaunchUIModule.h
@@ -0,0 +1,39 @@
+// Copyright Epic Games, Inc. All Rights Reserved.
+
+
+#include "CoreMinimal.h"
+#include "Modules/ModuleInterface.h"
+#include "Modules/ModuleManager.h"
+
+namespace CustomLaunchUI
+{
+	class ILaunchProfileMarshallerFactory;
+};
+
+/*
+ * Interface for CustomLaunchUI module
+ */
+class CUSTOMLAUNCHUI_API ICustomLaunchUIModule : public IModuleInterface
+{
+public:
+	/**
+	* Singleton-like access to the Custom Launch UI module instance
+	* @return Returns ICustomLaunchUIModule singleton instance, loading the module on demand if needed
+	*/
+	static inline ICustomLaunchUIModule& Get()
+	{
+		return FModuleManager::LoadModuleChecked<ICustomLaunchUIModule>("CustomLaunchUIModule");
+	}
+
+
+	/**
+	 * Registers a marhaller factory, for defining the layout & fields for editing an ILauncherProfile
+	 */
+	virtual void RegisterMarshaller( TSharedRef<CustomLaunchUI::ILaunchProfileMarshallerFactory> MarshallerFactory ) = 0;
+
+	/**
+	 * Unregisters a previously-registered marshaller factory
+	 */
+	virtual void UnregisterMarshaller( TSharedRef<CustomLaunchUI::ILaunchProfileMarshallerFactory> MarshallerFactory ) = 0;
+};
+
diff --git a/Engine/Plugins/Developer/CustomLaunchUI/Source/Public/Marshalling/BasicProfileMarshaller.h b/Engine/Plugins/Developer/CustomLaunchUI/Source/Public/Marshalling/BasicProfileMarshaller.h
new file mode 100644
--- /dev/null
+++ b/Engine/Plugins/Developer/CustomLaunchUI/Source/Public/Marshalling/BasicProfileMarshaller.h
@@ -0,0 +1,35 @@
+// Copyright Epic Games, Inc. All Rights Reserved.
+
+#pragma once
+
+#include "Marshalling/GenericProfileMarshaller.h"
+
+namespace CustomLaunchUI
+{
+
+	class FBasicProfileMarshaller : public FGenericProfileMarshaller
+	{
+	public:
+		FBasicProfileMarshaller( const ILauncherProfileRef& Profile, const TSharedRef<FModel>& InModel );
+		virtual ~FBasicProfileMarshaller() = default;
+
+		virtual void Construct() override;
+
+		virtual FString GetName() const override
+		{
+			return TEXT("BasicProfile");
+		}
+	};
+
+
+
+
+	class CUSTOMLAUNCHUI_API FBasicProfileMarshallerFactory : public ILaunchProfileMarshallerFactory
+	{
+	public:
+		virtual TSharedPtr<ILaunchProfileMarshaller> TryCreateMarshaller( const ILauncherProfileRef& InProfile, const TSharedRef<FModel>& InModel ) override;
+
+		virtual bool IsProfileTypeSupported(EProfileType ProfileType ) const override { return ProfileType == EProfileType::Basic; }
+
+	};
+}
diff --git a/Engine/Plugins/Developer/CustomLaunchUI/Source/Public/Marshalling/CustomProfileMarshaller.h b/Engine/Plugins/Developer/CustomLaunchUI/Source/Public/Marshalling/CustomProfileMarshaller.h
new file mode 100644
--- /dev/null
+++ b/Engine/Plugins/Developer/CustomLaunchUI/Source/Public/Marshalling/CustomProfileMarshaller.h
@@ -0,0 +1,33 @@
+// Copyright Epic Games, Inc. All Rights Reserved.
+
+#pragma once
+
+#include "Marshalling/GenericProfileMarshaller.h"
+
+namespace CustomLaunchUI
+{
+	class CUSTOMLAUNCHUI_API FCustomProfileMarshaller : public FGenericProfileMarshaller
+	{
+	public:
+		FCustomProfileMarshaller( const ILauncherProfileRef& Profile, const TSharedRef<FModel>& InModel );
+		virtual ~FCustomProfileMarshaller() = default;
+
+		virtual void Construct() override;
+
+		virtual FString GetName() const override
+		{
+			return TEXT("CustomProfile");
+		}
+	};
+
+
+	class CUSTOMLAUNCHUI_API FCustomProfileMarshallerFactory : public ILaunchProfileMarshallerFactory
+	{
+	public:
+		virtual TSharedPtr<ILaunchProfileMarshaller> TryCreateMarshaller( const ILauncherProfileRef& Profile, const TSharedRef<FModel>& InModel ) override;
+
+		virtual bool IsProfileTypeSupported(EProfileType ProfileType ) const override { return ProfileType == EProfileType::Custom; }
+
+	};
+
+}
diff --git a/Engine/Plugins/Developer/CustomLaunchUI/Source/Public/Marshalling/GenericProfileMarshaller.h b/Engine/Plugins/Developer/CustomLaunchUI/Source/Public/Marshalling/GenericProfileMarshaller.h
new file mode 100644
--- /dev/null
+++ b/Engine/Plugins/Developer/CustomLaunchUI/Source/Public/Marshalling/GenericProfileMarshaller.h
@@ -0,0 +1,177 @@
+// Copyright Epic Games, Inc. All Rights Reserved.
+
+#pragma once
+
+#include "Marshalling/ILaunchProfileMarshaller.h"
+#include "Model/CustomLaunchUIModel.h"
+
+class SSearchableComboBox;
+class SCustomLaunchMapListView;
+class SCustomLaunchDeviceListView;
+namespace ESelectInfo { enum Type : int; }
+
+namespace CustomLaunchUI
+{
+
+	class CUSTOMLAUNCHUI_API FGenericProfileMarshaller : public ILaunchProfileMarshaller, public TSharedFromThis<FGenericProfileMarshaller>
+	{
+	public:
+		FGenericProfileMarshaller( const ILauncherProfileRef& Profile, const ILauncherProfileRef& InDefaultProfile, const TSharedRef<FModel>& InModel );
+		virtual ~FGenericProfileMarshaller() = default;
+
+		virtual void Construct() override;
+
+		virtual FLaunchProfileTreeDataRef GetProfileTree() override
+		{
+			return TreeData;
+		}
+
+	protected:
+		// default property creation functions
+		void AddProjectProperty( FLaunchProfileTreeNode& HeadingNode );
+		void AddTargetProperty( FLaunchProfileTreeNode& HeadingNode );
+		void AddPlatformProperty( FLaunchProfileTreeNode& HeadingNode );
+		void AddConfigurationProperty( FLaunchProfileTreeNode& HeadingNode );
+		void AddContentSchemeProperty( FLaunchProfileTreeNode& HeadingNode );
+		void AddCompressPakFilesProperty( FLaunchProfileTreeNode& HeadingNode );
+		void AddUseIoStoreProperty( FLaunchProfileTreeNode& HeadingNode );
+		void AddGenerateChunksProperty( FLaunchProfileTreeNode& HeadingNode );
+		void AddImportZenSnapshotProperty( FLaunchProfileTreeNode& HeadingNode );
+		void AddZenPakStreamingPathProperty( FLaunchProfileTreeNode& HeadingNode );
+		void AddIncrementalCookProperty( FLaunchProfileTreeNode& HeadingNode );
+		void AddCookProperty( FLaunchProfileTreeNode& HeadingNode );
+		void AddMapsToCookProperty( FLaunchProfileTreeNode& HeadingNode );
+		void AddAdditionalCookerOptionsProperty( FLaunchProfileTreeNode& HeadingNode );
+		void AddBuildProperty( FLaunchProfileTreeNode& HeadingNode );
+		void AddForceBuildProperty( FLaunchProfileTreeNode& HeadingNode );
+		void AddBuidUATProperty( FLaunchProfileTreeNode& HeadingNode ); // fixme: not really relevant
+		void AddArchitectureProperty( FLaunchProfileTreeNode& HeadingNode );
+		void AddStagingDirectoryProperty( FLaunchProfileTreeNode& HeadingNode ); // fixme: not really relevant
+		void AddArchiveBuildProperty( FLaunchProfileTreeNode& HeadingNode );
+		void AddArchiveBuildDirectoryProperty( FLaunchProfileTreeNode& HeadingNode );
+		void AddDeployProperty( FLaunchProfileTreeNode& HeadingNode );
+		void AddIncrementalDeployProperty( FLaunchProfileTreeNode& HeadingNode ); // fixme: not really relavent
+		void AddTargetDeviceProperty( FLaunchProfileTreeNode& HeadingNode );
+		void AddRunProperty( FLaunchProfileTreeNode& HeadingNode );
+		void AddInitialMapProperty( FLaunchProfileTreeNode& HeadingNode );
+		void AddCommandLineProperty( FLaunchProfileTreeNode& HeadingNode );
+
+		// helper callbacks to simplify control enable/visibility
+		FLaunchProfileTreeNode::FGetBool ForPak;
+		FLaunchProfileTreeNode::FGetBool ForZenWS;
+		FLaunchProfileTreeNode::FGetBool ForCooked;
+		FLaunchProfileTreeNode::FGetBool ForEnabledCooked;
+		FLaunchProfileTreeNode::FGetBool ForCode;
+		FLaunchProfileTreeNode::FGetBool ForDeployment;
+		FLaunchProfileTreeNode::FGetBool ForRun;
+		FLaunchProfileTreeNode::FGetString EmptyString;
+
+
+
+		TArray<FString> CachedMapsToCook;
+		enum class EMapOption : uint8
+		{
+			Startup,
+			Selected,
+		};
+		EMapOption MapOption = EMapOption::Startup;
+
+		TArray<FString> CachedDeployDeviceIDs;
+		enum class EDeployDeviceOption : uint8
+		{
+			Default,
+			Selected,
+		};
+		EDeployDeviceOption DeployDeviceOption = EDeployDeviceOption::Default;
+
+
+		FString GetProjectPath() const;
+		void SetProjectName(FString ProjectPath);
+		bool HasProject() const;
+
+		FString GetBuildTarget() const;
+		void SetBuildTarget(FString BuildTarget);
+		TArray<EBuildTargetType> GetSupportedBuildTargetTypes() const;
+
+		void SetBuildConfiguration(EBuildConfiguration BuildConfiguration);
+		EBuildConfiguration GetBuildConfiguration() const;
+
+		void SetContentScheme(EContentScheme ContentScheme);
+		bool IsContentSchemeAvailable(EContentScheme, FText& OutReason) const;
+
+		FString GetCommandLine() const;
+		void SetCommandLine( const FString& NewCommandLine );
+
+		void SetSelectedPlatforms( TArray<FString> SelectedPlatforms );
+		TArray<FString> GetSelectedPlatforms() const;
+
+		void SetCook( bool bCook );
+		bool GetCook( ILauncherProfilePtr InProfile = nullptr ) const;
+
+		void SetMapsToCook(TArray<FString> MapsToCook);
+		TArray<FString> GetMapsToCook() const;
+		EMapOption GetMapOption() const;
+		void SetMapOption( EMapOption MapOption );
+		float GetMapListHeight() const;
+		void SetMapListHeight( float NewHeight );
+		TSharedRef<SWidget> CreateMapListWidget();
+		float MapListHeight = 300.0f;
+
+
+		void SetDeployDeviceIDs(TArray<FString> DeployDeviceIDs);
+		TArray<FString> GetDeployDeviceIDs() const;
+		EDeployDeviceOption GetDeployDeviceOption() const;
+		void SetDeployDeviceOption( EDeployDeviceOption DeployDeviceOption );
+		float GetDeployDeviceListHeight() const;
+		void SetDeployDeviceListHeight( float NewHeight );
+		void OnDeviceRemoved(FString DeviceID);
+		TSharedRef<SWidget> CreateDeployDeviceWidget();
+		float DeployDeviceListHeight = 150.0f;
+
+		void SetBuild( bool bBuild );
+		bool GetBuild( ILauncherProfilePtr InProfile = nullptr ) const;
+
+		void SetForceBuild( bool bForceBuild );
+		bool GetForceBuild( ILauncherProfilePtr InProfile = nullptr ) const;
+
+		void SetArchitecture( FString Architecture );
+		FString GetArchitecture() const;
+		FText GetArchitectureDisplayName( FString Architecture );
+
+		void SetDeployToDevice( bool bDeployToDevice );
+		bool GetDeployToDevice( ILauncherProfilePtr InProfile = nullptr ) const;
+
+		void SetIsRunning( bool bRun );
+		bool GetIsRunning( ILauncherProfilePtr InProfile = nullptr ) const;
+
+		void OnInitialMapChanged( TSharedPtr<FString> InitialMap, ESelectInfo::Type );
+		TSharedPtr<FString> GetInitialMap() const;
+
+
+
+	protected:
+
+		virtual void OnPropertyChanged() override;
+		void RefreshContentScheme();
+		void CacheStartupMapList();
+		void CacheArchitectures();
+		TSharedRef<SWidget> OnGenerateComboWidget( TSharedPtr<FString> InComboString );
+
+
+		FLaunchProfileTreeDataRef TreeData;
+		const ILauncherProfileRef Profile;
+		const ILauncherProfileRef DefaultProfile;
+		EProfileType ProfileType;
+
+		const TSharedRef<FModel> Model;
+
+		EContentScheme ContentScheme;
+		bool bShouldCook;
+		TArray<TSharedPtr<FString>> CachedStartupMaps;
+		TSharedPtr<SSearchableComboBox> InitalMapCombo;
+		TSharedPtr<SCustomLaunchMapListView> MapListView;
+		TSharedPtr<SCustomLaunchDeviceListView> DeployDeviceListView;
+		TArray<FString> CachedArchitectures;
+		EBuildTargetType CachedBuildTargetType = EBuildTargetType::Game;
+	};
+}
diff --git a/Engine/Plugins/Developer/CustomLaunchUI/Source/Public/Marshalling/ILaunchProfileMarshaller.h b/Engine/Plugins/Developer/CustomLaunchUI/Source/Public/Marshalling/ILaunchProfileMarshaller.h
new file mode 100644
--- /dev/null
+++ b/Engine/Plugins/Developer/CustomLaunchUI/Source/Public/Marshalling/ILaunchProfileMarshaller.h
@@ -0,0 +1,79 @@
+// Copyright Epic Games, Inc. All Rights Reserved.
+
+#pragma once
+
+#include "Marshalling/LaunchProfileTreeData.h"
+#include "Model/CustomLaunchUIModel.h"
+#include "ILauncherProfile.h"
+
+namespace CustomLaunchUI
+{
+
+	/**
+	 * Interface for a profile marshaller that creates FLaunchProfileTreeData from a given ILauncherProfile.
+	 * Expected to be created by an instance of ILaunchProfileMarshallerFactory
+	 */
+	class ILaunchProfileMarshaller
+	{
+	public:
+		/**
+		 * Construct the marshalling profile UI tree
+		 */
+		virtual void Construct() = 0;
+
+		/**
+		 * Provide access to the marshalling profile UI tree
+		 * 
+		 * @returns profile UI tree
+		 */
+		virtual FLaunchProfileTreeDataRef GetProfileTree() = 0;
+
+		/**
+		 * Debug name for this marshaller
+		 * 
+		 * @returns debug name of the marshaller
+		 */
+		virtual FString GetName() const = 0;
+
+		/**
+		 * Callback when the property tree editor has modified the profile.
+		 * NOTE: Custom widgets will need to call this manually.
+		 */
+		virtual void OnPropertyChanged() = 0;
+	};
+
+
+	/**
+	 * Interface for a factory that can create specializations of ILaunchProfileMarshaller for a given ILauncherProfile
+	 */
+	class ILaunchProfileMarshallerFactory
+	{
+	public:
+		/**
+		 * Create a profile marshalling class for the given profile
+		 * 
+		 * @param Profile the profile being edited
+		 * @param Model helper class
+		 * @returns new instance of a profile marshaller for the given profile, or null if it is not supported by this marshaller factory
+		 */
+		virtual TSharedPtr<ILaunchProfileMarshaller> TryCreateMarshaller( const ILauncherProfileRef& Profile, const TSharedRef<CustomLaunchUI::FModel>& Model ) = 0;
+
+		/**
+		 * Get the priority for this marshalling factory. Higher priority marshallers are evaluated first. 
+		 * Specializations should return non-zero values
+		 * 
+		 * @returns priority for this marshalling factory
+		 */
+		virtual int GetPriority() const { return 0; }
+
+		/**
+		 * Determines whether this marshaller can support profiles of the given type
+		 * 
+		 * @param ProfileType the type of profile
+		 * @returns true if the profile type is supported
+		 */
+		virtual bool IsProfileTypeSupported(CustomLaunchUI::EProfileType ProfileType ) const = 0;
+	};
+}
+
+
diff --git a/Engine/Plugins/Developer/CustomLaunchUI/Source/Public/Marshalling/LaunchProfileTreeData.h b/Engine/Plugins/Developer/CustomLaunchUI/Source/Public/Marshalling/LaunchProfileTreeData.h
new file mode 100644
--- /dev/null
+++ b/Engine/Plugins/Developer/CustomLaunchUI/Source/Public/Marshalling/LaunchProfileTreeData.h
@@ -0,0 +1,109 @@
+// Copyright Epic Games, Inc. All Rights Reserved.
+
+#pragma once
+
+#include "Internationalization/Text.h"
+#include "Templates/Function.h"
+#include "Templates/SharedPointer.h"
+#include "ILauncherProfile.h"
+
+
+class SWidget;
+struct EVisibility;
+namespace CustomLaunchUI
+{
+	class ILaunchProfileMarshaller;
+	class FLaunchProfileTreeData;
+	class FLaunchProfileTreeNode;
+	class FModel;
+
+	typedef TSharedPtr<FLaunchProfileTreeData> FLaunchProfileTreeDataPtr;
+	typedef TSharedRef<FLaunchProfileTreeData> FLaunchProfileTreeDataRef;
+
+	typedef TSharedPtr<FLaunchProfileTreeNode> FLaunchProfileTreeNodePtr;
+	typedef TSharedRef<FLaunchProfileTreeNode> FLaunchProfileTreeNodeRef;
+
+
+	class CUSTOMLAUNCHUI_API FLaunchProfileTreeNode : public TSharedFromThis<FLaunchProfileTreeNode>
+	{
+	public:
+		// construction
+		FLaunchProfileTreeNode( const FLaunchProfileTreeData* InTreeData );
+
+		typedef TFunction<void()> FFunction;
+		typedef TFunction<bool()> FGetBool;
+		typedef TFunction<void(bool)> FSetBool;
+		typedef TFunction<FString(void)> FGetString;
+		typedef TFunction<void(FString)> FSetString;
+
+		// node creation
+		struct FCallbacks
+		{
+			FGetBool IsDefault = nullptr;
+			FFunction SetToDefault = nullptr;
+			FGetBool IsVisible = nullptr;
+			FGetBool IsEnabled = nullptr;
+		};
+		FLaunchProfileTreeNode& AddWidget( FText InName, FCallbacks&& InWidgetCallbacks, TSharedRef<SWidget> InValueWidget );
+		FLaunchProfileTreeNode& AddWidget( FText InName, TSharedRef<SWidget> InValueWidget );
+
+
+		struct FBooleanCallbacks
+		{
+			FGetBool GetValue = nullptr;
+			FSetBool SetValue = nullptr;
+			FGetBool GetDefaultValue = nullptr;
+			FGetBool IsVisible = nullptr;
+			FGetBool IsEnabled = nullptr;
+		};
+		FLaunchProfileTreeNode& AddBoolean( FText InName, FBooleanCallbacks&& BooleanCallbacks );
+
+		struct FStringCallbacks
+		{
+			FGetString GetValue = nullptr;
+			FSetString SetValue = nullptr;
+			FGetString GetDefaultValue = nullptr;
+			FGetBool IsVisible = nullptr;
+			FGetBool IsEnabled = nullptr;
+
+		};
+		FLaunchProfileTreeNode& AddString( FText InName, FStringCallbacks&& StringCallbacks );
+		FLaunchProfileTreeNode& AddDirectoryString( FText InName, FStringCallbacks&& StringCallbacks );
+
+
+
+		// member variables
+		FText Name;
+		TSharedPtr<SWidget> Widget;
+
+		FCallbacks Callbacks;
+
+
+		TArray<FLaunchProfileTreeNodePtr> Children;
+		const FLaunchProfileTreeData* TreeData;
+
+	protected:
+
+	};
+
+
+
+	class CUSTOMLAUNCHUI_API FLaunchProfileTreeData : public TSharedFromThis<FLaunchProfileTreeData>
+	{
+	public:
+		// construction
+		FLaunchProfileTreeData(ILauncherProfilePtr InProfile, TSharedRef<CustomLaunchUI::FModel> InModel, ILaunchProfileMarshaller* InMarshaller);
+
+		// node creation
+		FLaunchProfileTreeNode& AddHeading( FText InName );
+
+		// member variables
+		ILauncherProfilePtr Profile;
+		TSharedRef<CustomLaunchUI::FModel> Model;
+		TArray<FLaunchProfileTreeNodePtr> Nodes;
+		ILaunchProfileMarshaller* Marshaller;
+	};
+
+
+}
+
diff --git a/Engine/Plugins/Developer/CustomLaunchUI/Source/Public/Model/CustomLaunchUIModel.h b/Engine/Plugins/Developer/CustomLaunchUI/Source/Public/Model/CustomLaunchUIModel.h
new file mode 100644
--- /dev/null
+++ b/Engine/Plugins/Developer/CustomLaunchUI/Source/Public/Model/CustomLaunchUIModel.h
@@ -0,0 +1,203 @@
+// Copyright Epic Games, Inc. All Rights Reserved.
+
+#pragma once
+
+#include "CoreTypes.h"
+#include "Containers/Array.h"
+#include "Delegates/Delegate.h"
+#include "ILauncherProfileManager.h"
+#include "Templates/SharedPointer.h"
+#include "Misc/DataDrivenPlatformInfoRegistry.h"
+
+
+
+namespace PlatformInfo
+{
+	struct FTargetPlatformInfo;
+};
+
+class ITargetDeviceProxyManager;
+class ITargetDeviceProxy;
+class ILauncher;
+struct FTargetInfo;
+
+namespace CustomLaunchUI
+{
+	enum class EContentScheme : uint8
+	{
+		PakFiles,
+		ZenStreaming,
+		ZenPakStreaming,
+		DevelopmentPackage,
+		LooseFiles,
+		CookOnTheFly,
+		MAX
+	};
+	TArray<EContentScheme> GetAllContentSchemes();
+	FText GetContentSchemeDisplayName(EContentScheme ContentScheme);
+	FText GetContentSchemeToolTip(EContentScheme ContentScheme);
+	FText GetProfileLaunchErrorMessage(ILauncherProfilePtr Profile);
+	const TCHAR* LexToString( const CustomLaunchUI::EContentScheme& ContentScheme);
+	bool LexTryParseString( CustomLaunchUI::EContentScheme& OutContentScheme, const TCHAR* String );
+
+	enum class EProfileType : uint8
+	{
+		Invalid,
+		Basic,
+		Custom,
+		Advanced,
+	};
+
+	struct FLaunchLogMessage
+	{
+		TSharedRef<FString> Message;
+		ELogVerbosity::Type Verbosity;
+
+		FLaunchLogMessage(const FString& InMessage, ELogVerbosity::Type InVerbosity)
+			: Message(MakeShared<FString>(InMessage))
+			, Verbosity(InVerbosity)
+		{}
+	};
+
+	struct FProjectSettings
+	{
+		// whether this project is the one that is currently open in the editor
+		bool bIsCurrentEditorProject = false;
+
+		// whether Zen Store is enabled via the project settings - in this case, we can't use Loose Files
+		bool bUseZenStore = false;
+
+		// whether a workspace will be created automatically when using Zen Pak streaming
+		bool bHasAutomaticZenPakStreamingWorkspaceCreation = false;
+
+		// whether Zen Server is allowed to accept connections from remote machines (i.e. console devkits, phones etc)
+		bool bAllowRemoteNetworkService = false;
+	};
+
+
+	DECLARE_DELEGATE_OneParam(FOnProfileClicked, const ILauncherProfilePtr&)
+
+	DECLARE_MULTICAST_DELEGATE_TwoParams(FOnSelectedProfileChanged, const ILauncherProfilePtr& /*new*/, const ILauncherProfilePtr& /*old*/)
+	DECLARE_MULTICAST_DELEGATE(FOnSelectedProjectChanged)
+	DECLARE_MULTICAST_DELEGATE(FOnBuildTargetChanged)
+	DECLARE_MULTICAST_DELEGATE(FOnSelectedPlatformChanged)
+
+	extern bool bUseFriendlyBuildTargetSelection;
+
+
+	class FModel : public TSharedFromThis<FModel>
+	{
+	public:
+		FModel(const TSharedRef<ITargetDeviceProxyManager>& InDeviceProxyManager, const TSharedRef<ILauncher>& InLauncher, const TSharedRef<ILauncherProfileManager>& InProfileManager);
+		~FModel();
+
+
+		const TSharedRef<ITargetDeviceProxyManager>& GetDeviceProxyManager() const
+		{
+			return DeviceProxyManager;
+		}
+
+		const TSharedRef<ILauncher>& GetLauncher() const
+		{
+			return Launcher;
+		}
+
+		const TSharedRef<ILauncherProfileManager>& GetProfileManager() const
+		{
+			return ProfileManager;
+		}
+
+
+		const TArray<ILauncherProfilePtr>& GetAllProfiles() const
+		{
+			return AllProfiles;
+		}
+
+		ILauncherProfilePtr GetBasicLaunchProfile() const
+		{
+			return BasicLaunchProfile;
+		}
+
+		const ILauncherProfilePtr& GetSelectedProfile() const
+		{
+			return SelectedProfile;
+		}
+
+		FOnSelectedProfileChanged& OnProfileSelected()
+		{
+			return ProfileSelectedDelegate;
+		}
+
+		void SelectProfile(const ILauncherProfilePtr& Profile);
+		void SortProfiles();
+
+		EProfileType GetProfileType(const ILauncherProfileRef& Profile) const;
+		bool IsAdvancedProfile(const ILauncherProfileRef& Profile) const;
+		bool IsBasicLaunchProfile( const ILauncherProfilePtr& Profile ) const;
+
+		ILauncherProfileRef CreateCustomProfile( const TCHAR* Name );
+		ILauncherProfileRef CreateBasicLaunchProfile();
+		ILauncherProfilePtr CloneCustomProfile(const ILauncherProfileRef& Profile);
+
+		const ILauncherProfileRef GetDefaultBasicLaunchProfile() const;
+		const ILauncherProfileRef GetDefaultCustomLaunchProfile() const;
+
+		EContentScheme DetermineProfileContentScheme(const ILauncherProfileRef& Profile) const;
+		void SetProfileContentScheme(EContentScheme ContentScheme, const ILauncherProfileRef& Profile, bool bWantToCook = true, ELauncherProfileDeploymentModes::Type DefaultDeploymentMode = ELauncherProfileDeploymentModes::CopyToDevice);
+
+		TSharedPtr<FLaunchLogMessage> AddLogMessage( const FString& InMessage, ELogVerbosity::Type InVerbosity );
+		void ClearLogMessages();
+		int32 GetNumLogMessages() const { return LaunchLogMessages.Num(); }
+
+
+		static const PlatformInfo::FTargetPlatformInfo* GetPlatformInfo(const ILauncherProfilePtr& Profile);
+		static const PlatformInfo::FTargetPlatformInfo* GetPlatformInfo(FName PlatformName, const FTargetInfo& BuildTargetInfo);
+
+		static bool IsHostPlatform(const ILauncherProfilePtr& Profile);
+		static bool IsHostPlatform(FName PlatformName);
+
+		static FTargetInfo GetBuildTargetInfo( const FString& BuildTargetName, const FString& ProjectPath );
+		static FTargetInfo GetBuildTargetInfo( const ILauncherProfileRef& Profile );
+
+		static TSharedPtr<ITargetDeviceProxy> GetDeviceProxy(const ILauncherProfileRef& Profile);
+		static void UpdatedCookedPlatformsFromDeployDeviceProxy(const ILauncherProfileRef& Profile, TSharedPtr<ITargetDeviceProxy> DeviceProxy = nullptr);
+		static void UpdateCookedPlatformsFromBuildTarget(const ILauncherProfileRef& Profile);
+
+		static FString GetVanillaPlatformName( const FString& PlatformName );
+		static FString GetBuildTargetPlatformName( const FString& PlatformName, const FTargetInfo& BuildTargetInfo );
+
+		const FProjectSettings GetProjectSettings( const FString& ProjectPath );
+		const FProjectSettings GetProjectSettings( const ILauncherProfileRef& Profile );
+
+	private:
+		void LoadConfig();
+		void SaveConfig();
+
+		void HandleProfileManagerProfileAdded(const ILauncherProfileRef& Profile);
+		void HandleProfileManagerProfileRemoved(const ILauncherProfileRef& Profile);
+		void HandleDeviceProxyAdded(const TSharedRef<ITargetDeviceProxy>& DeviceProxy );
+		void HandleDeviceProxyRemoved(const TSharedRef<ITargetDeviceProxy>& DeviceProxy );
+
+		void ReadProjectSettingsFromConfig( FConfigCacheIni& InConfig, FProjectSettings& OutResult );
+
+		TSharedRef<ITargetDeviceProxyManager> DeviceProxyManager;
+		TSharedRef<ILauncher> Launcher;
+		TSharedRef<ILauncherProfileManager> ProfileManager;
+
+		ILauncherProfilePtr SelectedProfile;
+		TArray<ILauncherProfilePtr> AllProfiles;
+		ILauncherProfilePtr BasicLaunchProfile;
+		bool bHasSetBasicLaunchProfilePlatform = false;
+
+		friend class FLaunchLogTextLayoutMarshaller;
+		TArray<TSharedPtr<FLaunchLogMessage>> LaunchLogMessages;
+
+		TMap<FString,FProjectSettings> CachedProjectSettings;
+
+		FOnSelectedProfileChanged ProfileSelectedDelegate;
+		FOnSelectedPlatformChanged PlatformChangedDelegate;
+
+		ILauncherProfilePtr DefaultBasicLaunchProfile;
+		ILauncherProfilePtr DefaultCustomLaunchProfile;
+	};
+}
diff --git a/Engine/Plugins/Developer/CustomLaunchUI/Source/Public/Widgets/Shared/SCustomLaunchBuildTargetCombo.h b/Engine/Plugins/Developer/CustomLaunchUI/Source/Public/Widgets/Shared/SCustomLaunchBuildTargetCombo.h
new file mode 100644
--- /dev/null
+++ b/Engine/Plugins/Developer/CustomLaunchUI/Source/Public/Widgets/Shared/SCustomLaunchBuildTargetCombo.h
@@ -0,0 +1,47 @@
+// Copyright Epic Games, Inc. All Rights Reserved.
+
+#pragma once
+
+#include "CoreMinimal.h"
+#include "Widgets/SCompoundWidget.h"
+#include "Styling/SlateTypes.h"
+#include "Styling/CoreStyle.h"
+#include "Styling/SlateWidgetStyleAsset.h"
+#include "Fonts/SlateFontInfo.h"
+
+class CUSTOMLAUNCHUI_API SCustomLaunchBuildTargetCombo
+	: public SCompoundWidget
+{
+public:
+	DECLARE_DELEGATE_OneParam(FOnSelectionChanged, FString );
+
+	SLATE_BEGIN_ARGS(SCustomLaunchBuildTargetCombo)
+		: _TextStyle( &FCoreStyle::Get().GetWidgetStyle<FTextBlockStyle>( "NormalText" ) )
+		{}
+		SLATE_EVENT(FOnSelectionChanged, OnSelectionChanged);
+		SLATE_ATTRIBUTE(FString, SelectedProject);
+		SLATE_ATTRIBUTE(FString, SelectedBuildTarget);
+		SLATE_ATTRIBUTE(TArray<EBuildTargetType>, SupportedTargetTypes);
+		SLATE_ATTRIBUTE(FSlateFontInfo, Font)
+		SLATE_STYLE_ARGUMENT( FTextBlockStyle, TextStyle )
+	SLATE_END_ARGS()
+
+public:
+	void Construct(	const FArguments& InArgs);
+
+protected:
+	FString GetDefaultBuildTargetName() const;
+
+	TAttribute<TArray<EBuildTargetType>> SupportedTargetTypes;
+	TAttribute<FString> SelectedProject;
+	TAttribute<FString> SelectedBuildTarget;
+	FOnSelectionChanged OnSelectionChanged;
+	bool bShowAnyProjectOption;
+
+	TSharedRef<SWidget> MakeBuildTargetSelectionWidget();
+
+	FText GetBuildTargetName() const;
+	void SetBuildTargetName(FString BuildTargetName);
+
+
+};
diff --git a/Engine/Plugins/Developer/CustomLaunchUI/Source/Public/Widgets/Shared/SCustomLaunchCombo.h b/Engine/Plugins/Developer/CustomLaunchUI/Source/Public/Widgets/Shared/SCustomLaunchCombo.h
new file mode 100644
--- /dev/null
+++ b/Engine/Plugins/Developer/CustomLaunchUI/Source/Public/Widgets/Shared/SCustomLaunchCombo.h
@@ -0,0 +1,144 @@
+// Copyright Epic Games, Inc. All Rights Reserved.
+
+#pragma once
+
+#include "CoreMinimal.h"
+#include "SlateOptMacros.h"
+#include "Widgets/SCompoundWidget.h"
+#include "Widgets/Input/SComboButton.h"
+#include "Styling/SlateTypes.h"
+#include "Styling/CoreStyle.h"
+#include "Framework/MultiBox/MultiBoxBuilder.h"
+
+// helper class to create a combo box that can select an item
+template<typename T>
+class SCustomLaunchCombo
+	: public SCompoundWidget
+{
+public:
+	DECLARE_DELEGATE_OneParam(FOnSelectionChanged, T );
+	DECLARE_DELEGATE_RetVal_OneParam(FText, FGetItemText, T)
+
+	SLATE_BEGIN_ARGS(SCustomLaunchCombo<T>)
+		{}
+		SLATE_EVENT(FOnSelectionChanged, OnSelectionChanged)
+		SLATE_EVENT(FGetItemText, GetDisplayName)
+		SLATE_EVENT(FGetItemText, GetItemToolTip)
+		SLATE_ATTRIBUTE(TArray<T>, Items)
+		SLATE_ATTRIBUTE(T, SelectedItem)
+	SLATE_END_ARGS()
+
+public:
+	void Construct(	const FArguments& InArgs)
+	{
+		OnSelectionChanged = InArgs._OnSelectionChanged;
+		GetDisplayName = InArgs._GetDisplayName;
+		GetItemToolTip = InArgs._GetItemToolTip;
+		Items = InArgs._Items;
+		SelectedItem = InArgs._SelectedItem;
+
+		ChildSlot
+		[
+			SNew(SComboButton)
+			.ButtonContent()
+			[
+				SNew(STextBlock)
+				.Text(this, &SCustomLaunchCombo<T>::GetSelectedItemDisplayName)
+				.Font(FCoreStyle::Get().GetFontStyle("SmallFont"))
+			]
+			.OnGetMenuContent(this, &SCustomLaunchCombo<T>::MakeWidget)
+		];
+	}
+
+protected:
+	TAttribute<T> SelectedItem;
+	TAttribute<TArray<T>> Items;
+	FOnSelectionChanged OnSelectionChanged;
+	FGetItemText GetDisplayName;
+	FGetItemText GetItemToolTip;
+
+	FText GetSelectedItemDisplayName() const
+	{
+		T Item = SelectedItem.Get();
+		return ToText(Item);
+	}
+
+	void SetSelectedItem(T Value)
+	{
+		OnSelectionChanged.ExecuteIfBound(Value);
+	}
+
+	FText ToText(T Value) const
+	{
+		if (GetDisplayName.IsBound())
+		{
+			return GetDisplayName.Execute(Value);
+		}
+
+		return ToTextOverride(Value);
+	}
+
+
+	virtual FText ToTextOverride(T Value) const
+	{
+		return FText::FromString("ERROR: GetDisplayName or ToTextOverride not bound");
+	}
+
+
+BEGIN_SLATE_FUNCTION_BUILD_OPTIMIZATION
+	TSharedRef<SWidget> MakeWidget()
+	{
+		FMenuBuilder MenuBuilder(true, nullptr);
+		{
+			TArray<T> MenuItems = Items.Get();
+
+			for (const T& Item : MenuItems)
+			{
+				FText ItemText = ToText(Item);
+				FText ToolTipText = GetItemToolTip.IsBound() ? GetItemToolTip.Execute(Item) : FText::GetEmpty();
+
+				MenuBuilder.AddMenuEntry(
+					ItemText, 
+					ToolTipText, 
+					FSlateIcon(), 
+					FUIAction(
+						FExecuteAction::CreateSP(this, &SCustomLaunchCombo<T>::SetSelectedItem, Item),
+						FCanExecuteAction(),
+						FGetActionCheckState::CreateLambda( [this, Item]() { return (SelectedItem.Get() == Item) ? ECheckBoxState::Checked : ECheckBoxState::Unchecked; })
+					),
+					NAME_None,
+					EUserInterfaceActionType::Check);
+			}
+		}
+
+		return MenuBuilder.MakeWidget();
+	}
+END_SLATE_FUNCTION_BUILD_OPTIMIZATION
+
+
+};
+
+
+// helper class that can display a list of items, if those items have a LexToString function available
+template<typename T>
+class SCustomLaunchLexToStringCombo : public SCustomLaunchCombo<T>
+{
+public:
+	virtual FText ToTextOverride(T Value) const override
+	{
+		return FText::FromString(LexToString(Value));
+	}
+};
+
+// helper class that can display a list of strings
+class SCustomLaunchStringCombo : public SCustomLaunchCombo<FString>
+{
+public:
+	virtual FText ToTextOverride(FString Value) const override
+	{
+		return FText::FromString(Value);
+	}
+};
+
+
+
diff --git a/Engine/Plugins/Developer/CustomLaunchUI/Source/Public/Widgets/Shared/SCustomLaunchContentSchemeCombo.h b/Engine/Plugins/Developer/CustomLaunchUI/Source/Public/Widgets/Shared/SCustomLaunchContentSchemeCombo.h
new file mode 100644
--- /dev/null
+++ b/Engine/Plugins/Developer/CustomLaunchUI/Source/Public/Widgets/Shared/SCustomLaunchContentSchemeCombo.h
@@ -0,0 +1,42 @@
+// Copyright Epic Games, Inc. All Rights Reserved.
+
+#pragma once
+
+#include "CoreMinimal.h"
+#include "Widgets/SCompoundWidget.h"
+#include "Styling/SlateTypes.h"
+#include "Styling/CoreStyle.h"
+#include "Styling/SlateWidgetStyleAsset.h"
+#include "Fonts/SlateFontInfo.h"
+#include "Model/CustomLaunchUIModel.h"
+
+class CUSTOMLAUNCHUI_API SCustomLaunchContentSchemeCombo
+	: public SCompoundWidget
+{
+public:
+	DECLARE_DELEGATE_OneParam(FOnSelectionChanged, CustomLaunchUI::EContentScheme );
+	DECLARE_DELEGATE_RetVal_TwoParams(bool, FGetContentSchemeAvailability, CustomLaunchUI::EContentScheme, FText& )
+
+	SLATE_BEGIN_ARGS(SCustomLaunchContentSchemeCombo) 
+		: _TextStyle( &FCoreStyle::Get().GetWidgetStyle<FTextBlockStyle>( "NormalText" ) )
+		{}
+		SLATE_EVENT(FOnSelectionChanged, OnSelectionChanged);
+		SLATE_EVENT(FGetContentSchemeAvailability, IsContentSchemeAvailable)
+		SLATE_ATTRIBUTE(CustomLaunchUI::EContentScheme, SelectedContentScheme);
+		SLATE_ATTRIBUTE(FSlateFontInfo, Font)
+		SLATE_STYLE_ARGUMENT( FTextBlockStyle, TextStyle )
+	SLATE_END_ARGS()
+
+public:
+	void Construct(	const FArguments& InArgs);
+
+protected:
+	TAttribute<CustomLaunchUI::EContentScheme> SelectedContentScheme;
+	FOnSelectionChanged OnSelectionChanged;
+	FGetContentSchemeAvailability IsContentSchemeAvailable;
+
+	TSharedRef<SWidget> MakeContentSchemeSelectionWidget();
+
+	FText GetContentSchemeName() const;
+	void SetContentScheme(CustomLaunchUI::EContentScheme ContentScheme);
+};
diff --git a/Engine/Plugins/Developer/CustomLaunchUI/Source/Public/Widgets/Shared/SCustomLaunchDeviceCombo.h b/Engine/Plugins/Developer/CustomLaunchUI/Source/Public/Widgets/Shared/SCustomLaunchDeviceCombo.h
new file mode 100644
--- /dev/null
+++ b/Engine/Plugins/Developer/CustomLaunchUI/Source/Public/Widgets/Shared/SCustomLaunchDeviceCombo.h
@@ -0,0 +1,37 @@
+// Copyright Epic Games, Inc. All Rights Reserved.
+
+#pragma once
+
+#include "CoreMinimal.h"
+#include "Widgets/Shared/SCustomLaunchDeviceWidgetBase.h"
+
+template<typename ItemType> class SComboBox;
+
+class CUSTOMLAUNCHUI_API SCustomLaunchDeviceCombo
+	: public SCustomLaunchDeviceWidgetBase
+{
+public:
+	SLATE_BEGIN_ARGS(SCustomLaunchDeviceCombo)
+		: _AllPlatforms(false)
+		{}
+		SLATE_EVENT(FOnDeviceRemoved, OnDeviceRemoved)
+		SLATE_EVENT(FOnSelectionChanged, OnSelectionChanged);
+		SLATE_ATTRIBUTE(TArray<FString>, SelectedDevices);
+		SLATE_ATTRIBUTE(TArray<FString>, Platforms);
+		SLATE_ARGUMENT(bool, AllPlatforms)
+	SLATE_END_ARGS()
+
+public:
+	void Construct(	const FArguments& InArgs );
+
+protected:
+
+	TSharedRef<SWidget> GenerateDeviceProxyListWidget(TSharedPtr<ITargetDeviceProxy> DeviceProxy) const;
+	void OnDeviceProxySelectionChangedChanged(TSharedPtr<ITargetDeviceProxy> DeviceProxy, ESelectInfo::Type InSelectInfo);
+	const FSlateBrush* GetSelectedDeviceProxyBrush() const;
+	FText GetSelectedDeviceProxyName() const;
+
+private:
+	TSharedPtr<SComboBox<TSharedPtr<ITargetDeviceProxy>> > DeviceProxyComboBox;
+
+};
diff --git a/Engine/Plugins/Developer/CustomLaunchUI/Source/Public/Widgets/Shared/SCustomLaunchDeviceListView.h b/Engine/Plugins/Developer/CustomLaunchUI/Source/Public/Widgets/Shared/SCustomLaunchDeviceListView.h
new file mode 100644
--- /dev/null
+++ b/Engine/Plugins/Developer/CustomLaunchUI/Source/Public/Widgets/Shared/SCustomLaunchDeviceListView.h
@@ -0,0 +1,42 @@
+// Copyright Epic Games, Inc. All Rights Reserved.
+
+#pragma once
+
+#include "CoreMinimal.h"
+#include "Widgets/Shared/SCustomLaunchDeviceWidgetBase.h"
+
+class ITableRow;
+class STableViewBase;
+template <typename ItemType> class SListView;
+
+class CUSTOMLAUNCHUI_API SCustomLaunchDeviceListView
+	: public SCustomLaunchDeviceWidgetBase
+{
+public:
+	SLATE_BEGIN_ARGS(SCustomLaunchDeviceListView)
+		: _AllPlatforms(false)
+		, _SingleSelect(false)
+		{}
+		SLATE_EVENT(FOnDeviceRemoved, OnDeviceRemoved)
+		SLATE_EVENT(FOnSelectionChanged, OnSelectionChanged);
+		SLATE_ATTRIBUTE(TArray<FString>, SelectedDevices);
+		SLATE_ATTRIBUTE(TArray<FString>, Platforms);
+		SLATE_ARGUMENT(bool, AllPlatforms)
+		SLATE_ARGUMENT(bool, SingleSelect)
+	SLATE_END_ARGS()
+
+public:
+	void Construct(	const FArguments& InArgs );
+
+protected:
+	bool bSingleSelect;
+
+	TSharedRef<ITableRow> GenerateDeviceProxyRow(TSharedPtr<ITargetDeviceProxy> DeviceProxy, const TSharedRef<STableViewBase>& OwnerTable);
+	ECheckBoxState IsDeviceProxyChecked(TSharedPtr<ITargetDeviceProxy> DeviceProxy) const;
+	void OnDeviceProxyCheckStateChanged(ECheckBoxState NewState, TSharedPtr<ITargetDeviceProxy> DeviceProxy);
+	virtual void OnDeviceListRefreshed() override;
+
+private:
+	TSharedPtr<SListView<TSharedPtr<ITargetDeviceProxy>> > DeviceProxyListView;
+
+};
diff --git a/Engine/Plugins/Developer/CustomLaunchUI/Source/Public/Widgets/Shared/SCustomLaunchDeviceWidgetBase.h b/Engine/Plugins/Developer/CustomLaunchUI/Source/Public/Widgets/Shared/SCustomLaunchDeviceWidgetBase.h
new file mode 100644
--- /dev/null
+++ b/Engine/Plugins/Developer/CustomLaunchUI/Source/Public/Widgets/Shared/SCustomLaunchDeviceWidgetBase.h
@@ -0,0 +1,40 @@
+// Copyright Epic Games, Inc. All Rights Reserved.
+
+#pragma once
+
+#include "CoreMinimal.h"
+#include "Widgets/SCompoundWidget.h"
+
+class ITargetDeviceProxy;
+class ITargetDeviceProxyManager;
+
+class CUSTOMLAUNCHUI_API SCustomLaunchDeviceWidgetBase
+	: public SCompoundWidget
+{
+public:
+	DECLARE_DELEGATE_OneParam(FOnSelectionChanged, TArray<FString> );
+	DECLARE_DELEGATE_OneParam(FOnDeviceRemoved, FString );
+
+public:
+	void Construct();
+	~SCustomLaunchDeviceWidgetBase();
+
+	void RefreshDeviceList();
+	void OnSelectedPlatformChanged();
+
+protected:
+	virtual void OnDeviceListRefreshed() {};
+
+	TAttribute<TArray<FString>> Platforms;
+	TAttribute<TArray<FString>> SelectedDevices;
+	FOnSelectionChanged OnSelectionChanged;
+	FOnDeviceRemoved OnDeviceRemoved;
+	bool bAllPlatforms = false;
+
+	void OnDeviceProxyAdded(const TSharedRef<ITargetDeviceProxy>& DeviceProxy);
+	void OnDeviceProxyRemoved(const TSharedRef<ITargetDeviceProxy>& DeviceProxy);
+
+	const TSharedRef<ITargetDeviceProxyManager> GetDeviceProxyManager() const;
+
+	TArray<TSharedPtr<ITargetDeviceProxy>> DeviceProxyList;
+};
diff --git a/Engine/Plugins/Developer/CustomLaunchUI/Source/Public/Widgets/Shared/SCustomLaunchMapListView.h b/Engine/Plugins/Developer/CustomLaunchUI/Source/Public/Widgets/Shared/SCustomLaunchMapListView.h
new file mode 100644
--- /dev/null
+++ b/Engine/Plugins/Developer/CustomLaunchUI/Source/Public/Widgets/Shared/SCustomLaunchMapListView.h
@@ -0,0 +1,70 @@
+// Copyright Epic Games, Inc. All Rights Reserved.
+
+#pragma once
+
+#include "CoreMinimal.h"
+#include "Widgets/SCompoundWidget.h"
+
+class ITableRow;
+class STableViewBase;
+template<typename ItemType> class STreeView;
+
+class CUSTOMLAUNCHUI_API SCustomLaunchMapListView
+	: public SCompoundWidget
+{
+public:
+	DECLARE_DELEGATE_OneParam(FOnSelectionChanged, TArray<FString> );
+
+	SLATE_BEGIN_ARGS(SCustomLaunchMapListView){}
+		SLATE_EVENT(FOnSelectionChanged, OnSelectionChanged);
+		SLATE_ATTRIBUTE(TArray<FString>, SelectedMaps);
+		SLATE_ATTRIBUTE(FString, ProjectPath)
+	SLATE_END_ARGS()
+
+public:
+	void Construct(	const FArguments& InArgs );
+
+	void RefreshMapList();
+	TSharedRef<SWidget> MakeControlsWidget();
+
+protected:
+	TAttribute<TArray<FString>> SelectedMaps;
+	TAttribute<FString> ProjectPath;
+	FOnSelectionChanged OnSelectionChanged;
+
+private:
+	bool bShowFolders = true; // whether to display the available maps in a hierarchy or flat list
+
+	struct FMapTreeNode
+	{
+		FString Name;
+		ECheckBoxState CheckBoxState = ECheckBoxState::Unchecked;
+		bool bFiltered = false;
+		TArray<TSharedPtr<FMapTreeNode>> Children;
+	};
+	typedef TSharedPtr<FMapTreeNode> FMapTreeNodePtr;
+
+	FMapTreeNodePtr MapTreeRoot;
+
+	void OnProjectChanged();
+
+
+	void RefreshCheckBoxState( bool bExpand = false);
+	ECheckBoxState RefreshCheckBoxStateRecursive(FMapTreeNodePtr Node, bool bExpand);
+	void SetCheckBoxStateRecursive(FMapTreeNodePtr Node, ECheckBoxState CheckBoxState, TArray<FString>& CheckedMaps);
+
+	void OnSearchFilterTextCommitted(const FText& SearchText, ETextCommit::Type InCommitType);
+	void OnSearchFilterTextChanged(const FText& SearchText);
+
+	FString CurrentFilterText;
+
+	TSharedPtr<STreeView<FMapTreeNodePtr>> MapTreeView;
+	TArray<FMapTreeNodePtr> MapTreeViewItemsSource;
+
+	void GetMapTreeNodeChildren(FMapTreeNodePtr Item, TArray<FMapTreeNodePtr>& OutChildren);
+	TSharedRef<ITableRow> GenerateMapTreeNodeRow(FMapTreeNodePtr Item, const TSharedRef<STableViewBase>& OwnerTable);
+	ECheckBoxState GetMapTreeNodeCheckState(FMapTreeNodePtr Item) const;
+	void SetMapTreeNodeCheckState(ECheckBoxState CheckBoxState, FMapTreeNodePtr Item);
+	const FSlateBrush* GetMapTreeNodeIcon(FMapTreeNodePtr Node) const;
+	FSlateColor GetMapTreeNodeColor(FMapTreeNodePtr Node) const;
+};
diff --git a/Engine/Plugins/Developer/CustomLaunchUI/Source/Public/Widgets/Shared/SCustomLaunchPlatformCombo.h b/Engine/Plugins/Developer/CustomLaunchUI/Source/Public/Widgets/Shared/SCustomLaunchPlatformCombo.h
new file mode 100644
--- /dev/null
+++ b/Engine/Plugins/Developer/CustomLaunchUI/Source/Public/Widgets/Shared/SCustomLaunchPlatformCombo.h
@@ -0,0 +1,44 @@
+// Copyright Epic Games, Inc. All Rights Reserved.
+
+#pragma once
+
+#include "CoreMinimal.h"
+#include "Widgets/SCompoundWidget.h"
+#include "Styling/CoreStyle.h"
+#include "Styling/SlateWidgetStyleAsset.h"
+#include "Fonts/SlateFontInfo.h"
+
+template<typename ItemType> class SComboBox;
+
+
+class CUSTOMLAUNCHUI_API SCustomLaunchPlatformCombo
+	: public SCompoundWidget
+{
+public:
+	DECLARE_DELEGATE_OneParam(FOnSelectionChanged, TArray<FString> );
+
+	SLATE_BEGIN_ARGS(SCustomLaunchPlatformCombo)
+		: _TextStyle( &FCoreStyle::Get().GetWidgetStyle<FTextBlockStyle>( "NormalText" ) )
+		{}
+		SLATE_EVENT(FOnSelectionChanged, OnSelectionChanged);
+		SLATE_ATTRIBUTE(TArray<FString>, SelectedPlatforms);
+		SLATE_ARGUMENT(bool, BasicPlatformsOnly)
+		SLATE_ATTRIBUTE(FSlateFontInfo, Font)
+		SLATE_STYLE_ARGUMENT( FTextBlockStyle, TextStyle )
+	SLATE_END_ARGS()
+
+public:
+	void Construct(	const FArguments& InArgs);
+
+protected:
+	TAttribute<TArray<FString>> SelectedPlatforms;
+	FOnSelectionChanged OnSelectionChanged;
+	bool bBasicPlatformsOnly;
+
+	TSharedRef<SWidget> OnGeneratePlatformListWidget( TSharedPtr<FString> Platform ) const;
+	void OnPlatformSelectionChanged( TSharedPtr<FString> Platform, ESelectInfo::Type InSelectInfo );
+	const FSlateBrush* GetSelectedPlatformBrush() const;
+	FText GetSelectedPlatformName() const;
+	TArray<TSharedPtr<FString>> PlatformsList;
+	TSharedPtr<SComboBox<TSharedPtr<FString>>> PlatformsComboBox;
+};
diff --git a/Engine/Plugins/Developer/CustomLaunchUI/Source/Public/Widgets/Shared/SCustomLaunchProjectCombo.h b/Engine/Plugins/Developer/CustomLaunchUI/Source/Public/Widgets/Shared/SCustomLaunchProjectCombo.h
new file mode 100644
--- /dev/null
+++ b/Engine/Plugins/Developer/CustomLaunchUI/Source/Public/Widgets/Shared/SCustomLaunchProjectCombo.h
@@ -0,0 +1,57 @@
+// Copyright Epic Games, Inc. All Rights Reserved.
+
+#pragma once
+
+#include "CoreMinimal.h"
+#include "Widgets/SCompoundWidget.h"
+#include "Styling/SlateTypes.h"
+#include "Styling/CoreStyle.h"
+#include "Styling/SlateWidgetStyleAsset.h"
+#include "Fonts/SlateFontInfo.h"
+
+class CUSTOMLAUNCHUI_API SCustomLaunchProjectCombo
+	: public SCompoundWidget
+{
+public:
+	DECLARE_DELEGATE_OneParam(FOnSelectionChanged, FString );
+
+	enum class ECurrentProjectOption
+	{
+		None,
+		ActualProject,
+		Empty,
+	};
+
+
+	SLATE_BEGIN_ARGS(SCustomLaunchProjectCombo)
+		: _ShowAnyProjectOption(false)
+		, _CurrentProjectOption(ECurrentProjectOption::None)
+		, _TextStyle( &FCoreStyle::Get().GetWidgetStyle<FTextBlockStyle>( "NormalText" ) )
+		{}
+		SLATE_EVENT(FOnSelectionChanged, OnSelectionChanged);
+		SLATE_ATTRIBUTE(FString, SelectedProject);
+		SLATE_ATTRIBUTE(bool, HasProject);
+		SLATE_ARGUMENT(bool, ShowAnyProjectOption)
+		SLATE_ARGUMENT(ECurrentProjectOption, CurrentProjectOption)
+		SLATE_ATTRIBUTE(FSlateFontInfo, Font)
+		SLATE_STYLE_ARGUMENT( FTextBlockStyle, TextStyle )
+	SLATE_END_ARGS()
+
+public:
+	void Construct(	const FArguments& InArgs);
+
+protected:
+	TAttribute<FString> SelectedProject;
+	TAttribute<bool> HasProject;
+	FOnSelectionChanged OnSelectionChanged;
+	bool bShowAnyProjectOption;
+	ECurrentProjectOption CurrentProjectOption;
+
+	TSharedRef<SWidget> MakeProjectSelectionWidget();
+	void OnBrowseForProject();
+
+	FText GetProjectName() const;
+	void SetProjectPath(FString ProjectPath);
+
+
+};
