remove prototype warning banner

#rnx

diff --git a/Engine/Plugins/Developer/CustomLaunchUI/Source/CustomLaunchUI/Private/Widgets/Profiles/SCustomLaunchProfilesPanel.cpp b/Engine/Plugins/Developer/CustomLaunchUI/Source/CustomLaunchUI/Private/Widgets/Profiles/SCustomLaunchProfilesPanel.cpp
--- a/Engine/Plugins/Developer/CustomLaunchUI/Source/CustomLaunchUI/Private/Widgets/Profiles/SCustomLaunchProfilesPanel.cpp
+++ b/Engine/Plugins/Developer/CustomLaunchUI/Source/CustomLaunchUI/Private/Widgets/Profiles/SCustomLaunchProfilesPanel.cpp
@@ -425,18 +425,6 @@ TSharedRef<SWidget> SCustomLaunchProfilesPanel::CreateProfileEditorToolbarWidget
 				SNew(SSpacer)
 			]
 
-			// prototype warning
-			+SHorizontalBox::Slot()
-			.AutoWidth()
-			.Padding(8)
-			.VAlign(VAlign_Top)
-			[
-				SNew(STextBlock)
-				.TextStyle(FAppStyle::Get(), "NormalText.Important")
-				.Text(LOCTEXT("PrototypeWarning", "This is a prototype and not ready for use!"))
-				.ColorAndOpacity(FAppStyle::Get().GetSlateColor("Colors.Warning"))
-			]
-
 			// launch button
 			+SHorizontalBox::Slot()
 			.AutoWidth()
