Add command line Extensibility system
- including Unreal Insights example

#rnx

diff --git a/Engine/Plugins/Developer/CustomLaunchUI/CustomLaunchUI.uplugin b/Engine/Plugins/Developer/CustomLaunchUI/CustomLaunchUI.uplugin
--- a/Engine/Plugins/Developer/CustomLaunchUI/CustomLaunchUI.uplugin
+++ b/Engine/Plugins/Developer/CustomLaunchUI/CustomLaunchUI.uplugin
@@ -32,6 +32,17 @@
 				"Mac"
 			],
 			"ProgramAllowList": [ "UnrealFrontend" ]
+		},
+		{
+			"Name": "CommonLaunchExtensions",
+			"Type": "Editor",
+			"LoadingPhase": "Default",
+			"PlatformAllowList": [
+				"Win64",
+				"Linux",
+				"Mac"
+			],
+			"ProgramAllowList": [ "UnrealFrontend" ]
 		}
 	]
 }
\ No newline at end of file
diff --git a/Engine/Plugins/Developer/CustomLaunchUI/Source/CommonLaunchExtensions/CommonLaunchExtensions.Build.cs b/Engine/Plugins/Developer/CustomLaunchUI/Source/CommonLaunchExtensions/CommonLaunchExtensions.Build.cs
new file mode 100644
--- /dev/null
+++ b/Engine/Plugins/Developer/CustomLaunchUI/Source/CommonLaunchExtensions/CommonLaunchExtensions.Build.cs
@@ -0,0 +1,19 @@
+// Copyright Epic Games, Inc. All Rights Reserved.
+
+using UnrealBuildTool;
+
+public class CommonLaunchExtensions : ModuleRules
+{
+	public CommonLaunchExtensions(ReadOnlyTargetRules Target) : base(Target)
+	{
+		PrivateDependencyModuleNames.AddRange(
+		[
+			"Core",
+			"Slate",
+			"SlateCore",
+			"Sockets",
+			"TraceLog",
+			"CustomLaunchUI",
+		]);
+	}
+}
diff --git a/Engine/Plugins/Developer/CustomLaunchUI/Source/CommonLaunchExtensions/Private/CommonLaunchExtensionsModule.cpp b/Engine/Plugins/Developer/CustomLaunchUI/Source/CommonLaunchExtensions/Private/CommonLaunchExtensionsModule.cpp
new file mode 100644
--- /dev/null
+++ b/Engine/Plugins/Developer/CustomLaunchUI/Source/CommonLaunchExtensions/Private/CommonLaunchExtensionsModule.cpp
@@ -0,0 +1,42 @@
+// Copyright Epic Games, Inc. All Rights Reserved.
+
+#include "Modules/ModuleInterface.h"
+#include "Templates/SharedPointer.h"
+#include "CustomLaunchUIModule.h"
+#include "Globals/GlobalsLaunchExtension.h"
+#include "Insights/InsightsLaunchExtension.h"
+
+class FCommonLaunchExtensionsModule : public IModuleInterface
+{
+public:
+
+	virtual void StartupModule() override
+	{
+		Globals = MakeShared<FGlobalsLaunchExtension>();
+		Insights = MakeShared<FInsightsLaunchExtension>();
+
+		ICustomLaunchUIModule::Get().RegisterExtension(Globals.ToSharedRef());
+		ICustomLaunchUIModule::Get().RegisterExtension(Insights.ToSharedRef());
+	}
+
+	virtual void ShutdownModule() override
+	{
+		if (ICustomLaunchUIModule* CustomLaunchUI = ICustomLaunchUIModule::TryGet())
+		{
+			CustomLaunchUI->UnregisterExtension(Globals.ToSharedRef());
+			CustomLaunchUI->UnregisterExtension(Insights.ToSharedRef());
+		}
+
+		Globals.Reset();
+		Insights.Reset();
+	}
+
+private:
+	TSharedPtr<FGlobalsLaunchExtension> Globals;
+	TSharedPtr<FInsightsLaunchExtension> Insights;
+
+};
+
+
+IMPLEMENT_MODULE(FCommonLaunchExtensionsModule, CommonLaunchExtensions);
+
diff --git a/Engine/Plugins/Developer/CustomLaunchUI/Source/CommonLaunchExtensions/Private/Globals/GlobalsLaunchExtension.cpp b/Engine/Plugins/Developer/CustomLaunchUI/Source/CommonLaunchExtensions/Private/Globals/GlobalsLaunchExtension.cpp
new file mode 100644
--- /dev/null
+++ b/Engine/Plugins/Developer/CustomLaunchUI/Source/CommonLaunchExtensions/Private/Globals/GlobalsLaunchExtension.cpp
@@ -0,0 +1,97 @@
+// Copyright Epic Games, Inc. All Rights Reserved.
+
+#include "Globals/GlobalsLaunchExtension.h"
+#include "SocketSubsystem.h"
+#include "Misc/OutputDeviceRedirector.h"
+#include "Trace/Trace.h"
+
+#define LOCTEXT_NAMESPACE "FGlobalsLaunchExtensionInstance"
+
+
+const TCHAR* FGlobalsLaunchExtensionInstance::LocalHostVariable     = TEXT("$(LocalHost)");
+const TCHAR* FGlobalsLaunchExtensionInstance::ProjectNameVariable   = TEXT("$(ProjectName)");
+const TCHAR* FGlobalsLaunchExtensionInstance::ProjectPathVariable   = TEXT("$(ProjectPath)");
+const TCHAR* FGlobalsLaunchExtensionInstance::TargetNameVariable    = TEXT("$(TargetName)");
+const TCHAR* FGlobalsLaunchExtensionInstance::PlatformNameVariable  = TEXT("$(Platform)");
+const TCHAR* FGlobalsLaunchExtensionInstance::ConfigurationVariable = TEXT("$(Configuration)");
+
+bool FGlobalsLaunchExtensionInstance::GetExtensionVariables( TArray<FString>& OutVariables ) const
+{
+	OutVariables.Add(LocalHostVariable);
+	OutVariables.Add(ProjectNameVariable);
+	OutVariables.Add(ProjectPathVariable);
+	OutVariables.Add(TargetNameVariable);
+	OutVariables.Add(PlatformNameVariable);
+	OutVariables.Add(ConfigurationVariable);
+	return true;
+}
+
+bool FGlobalsLaunchExtensionInstance::GetExtensionVariableValue( const FString& InVariable, FString& OutValue ) const
+{
+	if (InVariable == LocalHostVariable)
+	{
+		// @todo: what if we have multiple host addresses etc etc. ?
+		if (ISocketSubsystem* SocketSubsystem = ISocketSubsystem::Get(PLATFORM_SOCKETSUBSYSTEM))
+		{
+			bool bCanBindAll;
+			TSharedPtr<FInternetAddr> LocalHostAddr = SocketSubsystem->GetLocalHostAddr(*GLog, bCanBindAll);
+
+			if (LocalHostAddr->IsValid())
+			{
+				constexpr bool bAppendPort = false;
+				OutValue = LocalHostAddr->ToString(bAppendPort);
+				return true;
+			}
+		}
+
+		OutValue = TEXT("localhost");
+		return true;
+	}
+
+	else if (InVariable == ProjectNameVariable)
+	{
+		OutValue = GetProfile()->GetProjectName();
+		return true;
+	}
+	else if (InVariable == ProjectPathVariable)
+	{
+		OutValue = GetProfile()->GetProjectPath();
+		return true;
+	}
+	else if (InVariable == TargetNameVariable)
+	{
+		OutValue = GetProfile()->GetBuildTarget();
+		return true;
+	}
+	else if (InVariable == PlatformNameVariable)
+	{
+		OutValue = FString::Join( GetProfile()->GetCookedPlatforms(), TEXT("+") );
+		return true;
+	}
+	else if (InVariable == ConfigurationVariable)
+	{
+		OutValue = LexToString( GetProfile()->GetBuildConfiguration() );
+	}
+
+	return false;
+}
+
+
+
+TSharedPtr<CustomLaunchUI::FLaunchExtensionInstance> FGlobalsLaunchExtension::CreateInstanceForProfile( CustomLaunchUI::FLaunchExtensionInstance::FArgs& InArgs )
+{
+	return MakeShared<FGlobalsLaunchExtensionInstance>(InArgs);
+}
+
+const TCHAR* FGlobalsLaunchExtension::GetInternalName() const
+{
+	return TEXT("Globals");
+}
+
+FText FGlobalsLaunchExtension::GetDisplayName() const
+{
+	return LOCTEXT("ExtensionName", "Globals");
+}
+
+
+#undef LOCTEXT_NAMESPACE
diff --git a/Engine/Plugins/Developer/CustomLaunchUI/Source/CommonLaunchExtensions/Private/Globals/GlobalsLaunchExtension.h b/Engine/Plugins/Developer/CustomLaunchUI/Source/CommonLaunchExtensions/Private/Globals/GlobalsLaunchExtension.h
new file mode 100644
--- /dev/null
+++ b/Engine/Plugins/Developer/CustomLaunchUI/Source/CommonLaunchExtensions/Private/Globals/GlobalsLaunchExtension.h
@@ -0,0 +1,32 @@
+// Copyright Epic Games, Inc. All Rights Reserved.
+
+#pragma once
+
+#include "Extension/LaunchExtension.h"
+
+class FGlobalsLaunchExtensionInstance : public CustomLaunchUI::FLaunchExtensionInstance
+{
+public:
+	FGlobalsLaunchExtensionInstance( FArgs& InArgs ) : CustomLaunchUI::FLaunchExtensionInstance(InArgs) {}
+	virtual ~FGlobalsLaunchExtensionInstance() = default;
+
+	virtual bool GetExtensionVariables( TArray<FString>& OutVariables ) const override;
+	virtual bool GetExtensionVariableValue( const FString& InVariable, FString& OutValue ) const override;	
+
+private:
+
+	static const TCHAR* LocalHostVariable;
+	static const TCHAR* ProjectNameVariable;
+	static const TCHAR* ProjectPathVariable;
+	static const TCHAR* TargetNameVariable;
+	static const TCHAR* PlatformNameVariable;
+	static const TCHAR* ConfigurationVariable;	
+};
+
+class FGlobalsLaunchExtension : public CustomLaunchUI::FLaunchExtension
+{
+public:
+	virtual TSharedPtr<CustomLaunchUI::FLaunchExtensionInstance> CreateInstanceForProfile( CustomLaunchUI::FLaunchExtensionInstance::FArgs& InArgs ) override;
+	virtual const TCHAR* GetInternalName() const override;
+	virtual FText GetDisplayName() const override;
+};
\ No newline at end of file
diff --git a/Engine/Plugins/Developer/CustomLaunchUI/Source/CommonLaunchExtensions/Private/Insights/InsightsLaunchExtension.cpp b/Engine/Plugins/Developer/CustomLaunchUI/Source/CommonLaunchExtensions/Private/Insights/InsightsLaunchExtension.cpp
new file mode 100644
--- /dev/null
+++ b/Engine/Plugins/Developer/CustomLaunchUI/Source/CommonLaunchExtensions/Private/Insights/InsightsLaunchExtension.cpp
@@ -0,0 +1,167 @@
+// Copyright Epic Games, Inc. All Rights Reserved.
+
+#include "Insights/InsightsLaunchExtension.h"
+#include "Trace/Trace.h"
+
+
+#define LOCTEXT_NAMESPACE "FInsightsLaunchExtensionInstance"
+
+const TCHAR* FInsightsLaunchExtensionInstance::FileParam = TEXT("-tracefile");
+const TCHAR* FInsightsLaunchExtensionInstance::HostParam = TEXT("-tracehost=$(LocalHost)");
+const TCHAR* FInsightsLaunchExtensionInstance::TraceParam = TEXT("-trace=");
+const TCHAR* FInsightsLaunchExtensionInstance::StatNamedEventsParam = TEXT("-statnamedevents");
+
+
+
+bool FInsightsLaunchExtensionInstance::GetExtensionParameters( TArray<FString>& OutParameters ) const
+{
+	OutParameters.Add(HostParam);
+	OutParameters.Add(FileParam);
+	OutParameters.Add(StatNamedEventsParam);
+
+	return true;
+}
+
+FText FInsightsLaunchExtensionInstance::GetExtensionParameterDisplayName( const FString& InParameter ) const
+{
+	if (InParameter == HostParam)
+	{
+		return LOCTEXT("TraceHostLabel", "Trace to a computer");
+	}
+	else if (InParameter == FileParam)
+	{
+		return LOCTEXT("TraceFileLabel", "Trace to a file");
+	}
+	else if (InParameter == StatNamedEventsParam)
+	{
+		return LOCTEXT("TraceNamedEventsParam", "Capture named events");
+	}
+
+	return FLaunchExtensionInstance::GetExtensionParameterDisplayName(InParameter);
+
+}
+
+
+void FInsightsLaunchExtensionInstance::CacheTraceChannels()
+{
+	TraceChannels.Reset();
+
+	FString TraceParamValue = GetParameterValue(TraceParam);
+	TraceParamValue.ParseIntoArray(TraceChannels, TEXT(",") );
+
+	TraceChannels.Sort();
+}
+
+void FInsightsLaunchExtensionInstance::ToggleTraceChannel( const FString& InChannel )
+{
+	if (TraceChannels.Contains(InChannel))
+	{
+		TraceChannels.Remove(InChannel);
+	}
+	else
+	{
+		TraceChannels.Add(InChannel);
+	}
+
+	if (TraceChannels.IsEmpty())
+	{
+		RemoveParameter(TraceParam);
+	}
+	else
+	{
+		FString TraceParamValue = FString::Join( TraceChannels, TEXT(",") );
+		UpdateParameterValue( TraceParam, TraceParamValue );
+	}
+}
+
+bool FInsightsLaunchExtensionInstance::IsTraceChannelEnabled( const FString InChannel ) const
+{
+	return TraceChannels.Contains(InChannel);
+}
+
+
+void FInsightsLaunchExtensionInstance::CustomizeSubmenu( FMenuBuilder& MenuBuilder )
+{
+	auto ChannelsMenuDelegate = [this]( FMenuBuilder& MenuBuilder )
+	{
+		CacheTraceChannels();
+
+		// collect all trace channels
+		TArray<FString> AllTraceChannels;
+		UE::Trace::EnumerateChannels([](const UE::Trace::FChannelInfo& ChannelInfo, void* User)
+		{
+			TArray<FString>& AllTraceChannels = *static_cast<TArray<FString>*>(User);
+			if (!ChannelInfo.bIsReadOnly)
+			{
+				FString Channel = FString(ChannelInfo.Name).LeftChop(7); // Remove "Channel" suffix
+				AllTraceChannels.Add(Channel);
+			}
+			return true;
+		}, &AllTraceChannels );
+		AllTraceChannels.Sort();
+
+		// add menu item to clear the selected channels
+		MenuBuilder.AddMenuEntry(
+			LOCTEXT("ClearTraceChannelsLabel", "None"),
+			FText::GetEmpty(),
+			FSlateIcon(),
+			FUIAction( 
+				FExecuteAction::CreateLambda( [this]() { RemoveParameter(TraceParam); TraceChannels.Reset(); } ),
+				FCanExecuteAction(),
+				FIsActionChecked::CreateLambda( [this]() { return (TraceChannels.IsEmpty()); } )
+				),
+				NAME_None,
+				EUserInterfaceActionType::ToggleButton);
+		MenuBuilder.AddMenuSeparator();
+
+		// add submenu items for each channel
+		for ( const FString& Channel : AllTraceChannels )
+		{
+			MenuBuilder.AddMenuEntry(
+				FText::FromString(Channel),
+				FText::GetEmpty(),
+				FSlateIcon(),
+				FUIAction(
+					FExecuteAction::CreateLambda( [this, Channel]() { ToggleTraceChannel(Channel); }),
+					FCanExecuteAction(),
+					FIsActionChecked::CreateLambda( [this, Channel]() { return IsTraceChannelEnabled(Channel); } )
+				),
+				NAME_None,
+				EUserInterfaceActionType::ToggleButton);
+		}
+	};
+
+
+
+	MenuBuilder.AddSubMenu(
+		LOCTEXT("TraceChannelLabels", "Select Channels"),
+		FText::GetEmpty(),
+		FNewMenuDelegate::CreateLambda(ChannelsMenuDelegate),
+		true, // bInOpenSubMenuOnClick
+		FSlateIcon(),
+		false // bInShouldCloseWindowAfterMenuSelection
+	);
+}
+
+
+
+
+
+
+TSharedPtr<CustomLaunchUI::FLaunchExtensionInstance> FInsightsLaunchExtension::CreateInstanceForProfile( CustomLaunchUI::FLaunchExtensionInstance::FArgs& InArgs )
+{
+	return MakeShared<FInsightsLaunchExtensionInstance>(InArgs);
+}
+
+const TCHAR* FInsightsLaunchExtension::GetInternalName() const
+{
+	return TEXT("Insights");
+}
+
+FText FInsightsLaunchExtension::GetDisplayName() const
+{
+	return LOCTEXT("ExtensionName", "Unreal Insights");
+}
+
+
+#undef LOCTEXT_NAMESPACE
diff --git a/Engine/Plugins/Developer/CustomLaunchUI/Source/CommonLaunchExtensions/Private/Insights/InsightsLaunchExtension.h b/Engine/Plugins/Developer/CustomLaunchUI/Source/CommonLaunchExtensions/Private/Insights/InsightsLaunchExtension.h
new file mode 100644
--- /dev/null
+++ b/Engine/Plugins/Developer/CustomLaunchUI/Source/CommonLaunchExtensions/Private/Insights/InsightsLaunchExtension.h
@@ -0,0 +1,41 @@
+// Copyright Epic Games, Inc. All Rights Reserved.
+
+#pragma once
+
+#include "Extension/LaunchExtension.h"
+
+
+
+class FInsightsLaunchExtensionInstance : public CustomLaunchUI::FLaunchExtensionInstance
+{
+public:
+	FInsightsLaunchExtensionInstance( FArgs& InArgs ) : FLaunchExtensionInstance(InArgs) {};
+	virtual ~FInsightsLaunchExtensionInstance() = default;
+
+	virtual bool GetExtensionParameters( TArray<FString>& OutParameters ) const override;
+	virtual FText GetExtensionParameterDisplayName( const FString& InParameter ) const override;
+	
+	virtual void CustomizeSubmenu( FMenuBuilder& MenuBuilder ) override;
+
+	void CacheTraceChannels();
+	void ToggleTraceChannel( const FString& InChannel );
+	bool IsTraceChannelEnabled( const FString InChannel ) const;
+
+private:
+	TArray<FString> TraceChannels;
+
+	static const TCHAR* FileParam;
+	static const TCHAR* HostParam;
+	static const TCHAR* TraceParam;
+	static const TCHAR* StatNamedEventsParam;
+
+};
+
+
+class FInsightsLaunchExtension : public CustomLaunchUI::FLaunchExtension
+{
+public:
+	virtual TSharedPtr<CustomLaunchUI::FLaunchExtensionInstance> CreateInstanceForProfile( CustomLaunchUI::FLaunchExtensionInstance::FArgs& InArgs ) override;
+	virtual const TCHAR* GetInternalName() const override;
+	virtual FText GetDisplayName() const override;
+};
\ No newline at end of file
diff --git a/Engine/Plugins/Developer/CustomLaunchUI/Source/CustomLaunchUI/CustomLaunchUI.Build.cs b/Engine/Plugins/Developer/CustomLaunchUI/Source/CustomLaunchUI/CustomLaunchUI.Build.cs
--- a/Engine/Plugins/Developer/CustomLaunchUI/Source/CustomLaunchUI/CustomLaunchUI.Build.cs
+++ b/Engine/Plugins/Developer/CustomLaunchUI/Source/CustomLaunchUI/CustomLaunchUI.Build.cs
@@ -10,13 +10,10 @@ public class CustomLaunchUI : ModuleRules
 		[
 			"Core",
 			"CoreUObject",
-			"LauncherServices",
 			"TargetPlatform",
 			"DesktopPlatform",
 			"ApplicationCore",
 			"InputCore",
-			"Slate",
-			"SlateCore",
 			"WorkspaceMenuStructure",
 			"ToolWidgets",
 			"TargetDeviceServices",
@@ -24,5 +21,12 @@ public class CustomLaunchUI : ModuleRules
 			"DeveloperToolSettings",
 			"Zen",
 		]);
+
+		PublicDependencyModuleNames.AddRange(
+		[
+			"Slate",
+			"SlateCore",
+			"LauncherServices",
+		]);
 	}
 }
diff --git a/Engine/Plugins/Developer/CustomLaunchUI/Source/CustomLaunchUI/Private/CustomLaunchUIModule.cpp b/Engine/Plugins/Developer/CustomLaunchUI/Source/CustomLaunchUI/Private/CustomLaunchUIModule.cpp
--- a/Engine/Plugins/Developer/CustomLaunchUI/Source/CustomLaunchUI/Private/CustomLaunchUIModule.cpp
+++ b/Engine/Plugins/Developer/CustomLaunchUI/Source/CustomLaunchUI/Private/CustomLaunchUIModule.cpp
@@ -22,6 +22,9 @@ namespace CustomLaunchUI
 {
 	extern void RegisterMarshallerFactory( TSharedRef<ILaunchProfileMarshallerFactory> MarshallerFactory );
 	extern void UnregisterMarshallerFactory( TSharedRef<ILaunchProfileMarshallerFactory> MarshallerFactory );
+
+	extern void RegisterExtension( TSharedRef<FLaunchExtension> Extension );
+	extern void UnregisterExtension( TSharedRef<FLaunchExtension> Extension );
 }
 
 
@@ -64,6 +67,17 @@ public:
 	}
 
 
+	virtual void RegisterExtension( TSharedRef<CustomLaunchUI::FLaunchExtension> Extension ) override
+	{
+		CustomLaunchUI::RegisterExtension(Extension);
+	}
+
+	virtual void UnregisterExtension( TSharedRef<CustomLaunchUI::FLaunchExtension> Extension ) override
+	{
+		CustomLaunchUI::UnregisterExtension(Extension);
+	}
+
+
 
 
 private:
diff --git a/Engine/Plugins/Developer/CustomLaunchUI/Source/CustomLaunchUI/Private/Extension/LaunchExtension.cpp b/Engine/Plugins/Developer/CustomLaunchUI/Source/CustomLaunchUI/Private/Extension/LaunchExtension.cpp
new file mode 100644
--- /dev/null
+++ b/Engine/Plugins/Developer/CustomLaunchUI/Source/CustomLaunchUI/Private/Extension/LaunchExtension.cpp
@@ -0,0 +1,489 @@
+// Copyright Epic Games, Inc. All Rights Reserved.
+
+#include "Extension/LaunchExtension.h"
+#include "Model/CustomLaunchUIModel.h"
+#include "Marshalling/ILaunchProfileMarshaller.h"
+#include "Containers/Array.h"
+#include "Templates/SharedPointer.h"
+#include "HAL/CriticalSection.h"
+#include "Misc/ScopeLock.h"
+#include "Misc/Parse.h"
+#include "Misc/ConfigCacheIni.h"
+#include "Modules/ModuleManager.h"
+#include "ILauncherServicesModule.h"
+
+#define LOCTEXT_NAMESPACE "FLaunchExtensionInstance"
+
+
+namespace CustomLaunchUI
+{
+	static TArray<TSharedPtr<FLaunchExtension>> GExtensions;
+	static FCriticalSection GExtensionsCS;
+
+
+	void RegisterExtension( TSharedRef<FLaunchExtension> Extension )
+	{
+		FScopeLock Lock(&GExtensionsCS);
+		GExtensions.Add(Extension);
+	}
+
+	void UnregisterExtension( TSharedRef<FLaunchExtension> Extension )
+	{
+		FScopeLock Lock(&GExtensionsCS);
+		GExtensions.Remove(Extension);
+	}
+
+
+	void ApplyExtensionVariables( const ILauncherProfileRef& InProfile, FString& InOutCommandLine, TSharedRef<FModel> InModel )
+	{
+		// take a snapshot of the extension list
+		TArray<TSharedPtr<FLaunchExtension>> Extensions;
+		{
+			FScopeLock Lock(&GExtensionsCS);
+			Extensions = GExtensions;
+		}
+
+
+		for (TSharedPtr<FLaunchExtension> Extension : Extensions)
+		{
+			// instantiate the extension and apply the extension variables
+			FLaunchExtensionInstance::FArgs Args
+			{
+				.Profile = InProfile,
+				.Marshaller = nullptr,
+				.Model = InModel,
+				.Extension = Extension.ToSharedRef(),
+			};
+			TSharedPtr<FLaunchExtensionInstance> ExtensionInstance = Extension->CreateInstanceForProfile(Args);
+			if (ExtensionInstance.IsValid())
+			{
+				// apply variable substitutions
+				TArray<FString> Variables;
+				if (ExtensionInstance->GetExtensionVariables(Variables))
+				{
+					for (const FString& Variable : Variables)
+					{
+						if (InOutCommandLine.Contains(Variable, ESearchCase::IgnoreCase))
+						{
+							FString VariableValue;
+							if (ExtensionInstance->GetExtensionVariableValue( Variable, VariableValue ))
+							{
+								InOutCommandLine.ReplaceInline( *Variable, *VariableValue, ESearchCase::IgnoreCase );
+							}
+						}
+					}
+				}
+
+				// allow for advanced command line customization
+				ExtensionInstance->CustomizeLaunchCommandLine( InOutCommandLine );
+			}
+		}
+	}
+
+
+
+
+	TArray<TSharedPtr<FLaunchExtensionInstance>> FLaunchExtension::CreateExtensionInstancesForProfile( ILauncherProfileRef InProfile, TSharedRef<FModel> InModel, ILaunchProfileMarshaller* InMarshaller )
+	{
+		// take a snapshot of the extension list
+		TArray<TSharedPtr<FLaunchExtension>> Extensions;
+		{
+			FScopeLock Lock(&GExtensionsCS);
+			Extensions = GExtensions;
+		}
+
+		// attempt to instantiate all extensions
+		TArray<TSharedPtr<FLaunchExtensionInstance>> Result;
+		for (TSharedPtr<FLaunchExtension> Extension : Extensions)
+		{
+			FLaunchExtensionInstance::FArgs Args
+			{
+				.Profile = InProfile,
+				.Marshaller = InMarshaller,
+				.Model = InModel,
+				.Extension = Extension.ToSharedRef(),
+			};
+
+			TSharedPtr<FLaunchExtensionInstance> ExtensionInstance = Extension->CreateInstanceForProfile(Args);
+			if (ExtensionInstance.IsValid())
+			{
+				Result.Add(ExtensionInstance);
+			}
+		}
+
+		return Result;
+	}
+
+
+
+
+
+
+
+
+
+
+
+	FLaunchExtensionInstance::FLaunchExtensionInstance( FArgs& InArgs )
+		: Profile(InArgs.Profile)
+		, Marshaller(InArgs.Marshaller)
+		, Model(InArgs.Model)
+		, Extension(InArgs.Extension)
+	{
+	}
+
+	FLaunchExtensionInstance::~FLaunchExtensionInstance()
+	{
+	}
+
+	bool FLaunchExtensionInstance::GetExtensionVariables( TArray<FString>& OutVariables ) const
+	{
+		return false;
+	}
+
+	bool FLaunchExtensionInstance::GetExtensionVariableValue( const FString& InParameter, FString& OutValue ) const
+	{
+		return false;
+	}
+
+	FText FLaunchExtensionInstance::GetExtensionParameterDisplayName( const FString& InParameter ) const
+	{
+		return FText::FromString(InParameter);
+	}
+
+	bool FLaunchExtensionInstance::GetExtensionParameters( TArray<FString>& OutParameters ) const
+	{
+		return false;
+	}
+
+	FString FLaunchExtensionInstance::GetParameterValue( const FString& InParameter ) const
+	{
+		FString CommandLine = GetCommandLine();
+
+		// InParameter is -Key=
+		FString ParamValue;
+		if (FParse::Value(*CommandLine, *InParameter, ParamValue, false))
+		{
+			return ParamValue;
+		}
+
+		// InParameter is -Key=Value
+		FString ParamKey;
+		if (InParameter.Split(TEXT("="), &ParamKey, &ParamValue))
+		{
+			ParamKey += TEXT("=");
+			if (FParse::Value(*CommandLine, *ParamKey, ParamValue, false))
+			{
+				return ParamValue;
+			}
+		}
+
+		return FString();
+	}
+
+	bool FLaunchExtensionInstance::UpdateParameterValue( const FString& InParameter, const FString& NewValue )
+	{
+		FString ParamKey, ParamValue;
+		if (InParameter.Split(TEXT("="), &ParamKey, &ParamValue))
+		{
+			RemoveParameter(InParameter);
+
+			FString NewParameter = ParamKey + "=" + NewValue;
+			AddParameter(NewParameter);
+			return true;
+		}
+
+		return false;
+
+	}
+
+
+
+	FString FLaunchExtensionInstance::GetFinalParameter( const FString& InParameter ) const
+	{
+		FString CommandLine = GetCommandLine();
+
+		// get the parameter's key & value
+		FString ParamName, ParamValue;
+		bool bParameterHasValue = InParameter.Split(TEXT("="), &ParamName, &ParamValue);
+		if (!bParameterHasValue)
+		{
+			ParamName = InParameter;
+		}
+
+		// the parameter's value may have been added or alterered - return how it is now
+		FString ParamKey = ParamName + TEXT("=");
+		if (FParse::Value(*CommandLine, *ParamKey, ParamValue, false))
+		{
+			return ParamKey + ParamValue;
+		}
+		
+		// the parameter may not have a value or it has been removed - return how it is now
+		FString Param = ParamName;
+		if (Param.RemoveFromStart(TEXT("-")) && FParse::Param(*CommandLine, *Param))
+		{
+			return ParamName;
+		}
+
+		// return the parameter as-is
+		return InParameter;
+	}
+
+	void FLaunchExtensionInstance::AddParameter( const FString& InParameter )
+	{
+		FString CommandLine = GetCommandLine() + TEXT(" ") + InParameter;
+		SetCommandLine(CommandLine);
+	}
+
+	void FLaunchExtensionInstance::RemoveParameter( const FString& InParameter )
+	{
+		if (TryRemoveParameterGroup(InParameter))
+		{
+			return;
+		}
+
+		FString Parameter = GetFinalParameter(InParameter);
+		FString CommandLine = GetCommandLine();
+
+		// first try to remove the parameter and the preceding space
+		FString ParameterWithSpace = TEXT(" ") + Parameter;
+		if (CommandLine.ReplaceInline(*ParameterWithSpace, TEXT(""), ESearchCase::IgnoreCase) > 0)
+		{
+			SetCommandLine(CommandLine);
+			return;
+		}
+
+		// next try to remove the parameter with a trailing space
+		ParameterWithSpace = Parameter + TEXT(" ");
+		if (CommandLine.ReplaceInline(*ParameterWithSpace, TEXT(""), ESearchCase::IgnoreCase) > 0)
+		{
+			SetCommandLine(CommandLine);
+			return;
+		}
+
+		// a little unexpected... just remove the parameter on its own
+		if (CommandLine.ReplaceInline(*Parameter, TEXT(""), ESearchCase::IgnoreCase) > 0)
+		{
+			SetCommandLine(CommandLine);
+			return;
+		}
+	}
+
+
+	bool FLaunchExtensionInstance::IsParameterUsed( const FString& InParameter ) const
+	{
+		FString Parameter = GetFinalParameter(InParameter);
+		return GetCommandLine().Contains(Parameter, ESearchCase::IgnoreCase);
+	}
+
+	void FLaunchExtensionInstance::SetParameterUsed( const FString& InParameter, bool bUsed )
+	{
+		bool bIsUsed = IsParameterUsed(InParameter);
+		if (bIsUsed != bUsed)
+		{
+			if (bIsUsed)
+			{
+				RemoveParameter(InParameter);
+			}
+			else
+			{
+				AddParameter(InParameter);
+			}
+		}
+	}
+
+
+	bool FLaunchExtensionInstance::IsParameterGroup( const FString& InParameter ) const
+	{
+		// see if InParameter contains just a simple -Key=Value or -Param
+		FString ParsedParameter = InParameter;
+		FString ParamKey, ParamValue;
+		if (InParameter.Split(TEXT("="), &ParamKey, &ParamValue))
+		{
+			ParamKey += TEXT("=");
+			if (FParse::Value(*InParameter, *ParamKey, ParamValue, false))
+			{
+				ParsedParameter = ParamKey + ParamValue;
+			}
+		}
+		
+		return (ParsedParameter.TrimStartAndEnd() != InParameter.TrimStartAndEnd());
+	}
+
+	bool FLaunchExtensionInstance::TryRemoveParameterGroup( const FString& InParameter )
+	{
+		if (!IsParameterGroup(InParameter))
+		{
+			return false;
+		}
+
+		// it may contain a group - handle them separately
+		const TCHAR* Ptr = InParameter.GetCharArray().GetData();
+		FString SubParameter = FParse::Token(Ptr, false);
+		while (!SubParameter.IsEmpty())
+		{
+			RemoveParameter(SubParameter);
+			SubParameter = FParse::Token(Ptr, false);
+		}
+
+		return true;
+	}
+		
+
+
+
+	void FLaunchExtensionInstance::MakeCommandLineSubmenu( FMenuBuilder& MenuBuilder)
+	{
+		auto ToggleParameter = [this]( FString Parameter )
+		{
+			if (IsParameterUsed(Parameter))
+			{
+				RemoveParameter(Parameter);
+			}
+			else
+			{
+				AddParameter(Parameter);
+			}
+		};
+
+		TArray<FString> Parameters;
+		if (GetExtensionParameters(Parameters) && Parameters.Num() > 0)
+		{
+			for (const FString& Parameter : Parameters)
+			{
+				FText ParameterDisplayName = GetExtensionParameterDisplayName(Parameter);
+				FText ParameterToolTip = FText::FromString(Parameter);
+
+				MenuBuilder.AddMenuEntry( 
+					ParameterDisplayName,
+					ParameterToolTip,
+					FSlateIcon(),
+					FUIAction(
+						FExecuteAction::CreateLambda(ToggleParameter, Parameter),
+						FCanExecuteAction(),
+						FIsActionChecked::CreateLambda( [this, Parameter]() { return IsParameterUsed(Parameter); } )
+					),
+					NAME_None,
+					EUserInterfaceActionType::Check);
+			}
+		}
+
+		CustomizeSubmenu(MenuBuilder);
+	}
+
+	FString FLaunchExtensionInstance::GetCommandLine() const
+	{
+		return Profile->GetAdditionalCommandLineParameters();
+	}
+
+	void FLaunchExtensionInstance::SetCommandLine( const FString& CommandLine )
+	{
+		Profile->SetAdditionalCommandLineParameters(CommandLine);
+		Marshaller->OnPropertyChanged();
+	}
+
+
+
+
+	FString FLaunchExtensionInstance::GetConfigString( EConfig Config, const TCHAR* Name, const TCHAR* DefaultValue ) const
+	{
+		FString KeyName = GetConfigKeyName(Config, Name);
+
+		if (Config == EConfig::PerProfile)
+		{
+			const FString* ValuePtr = Profile->GetCustomStringProperties().Find( KeyName );
+			return ValuePtr ? *ValuePtr : DefaultValue;
+		}
+		else
+		{
+			FString Value;
+			return GConfig->GetString( Model->GetConfigSection() , *KeyName, Value, Model->GetConfigIni() ) ? Value : DefaultValue;
+		}
+	}
+
+	bool FLaunchExtensionInstance::GetConfigBool( EConfig Config, const TCHAR* Name, bool DefaultValue ) const
+	{
+		FString KeyName = GetConfigKeyName(Config, Name);
+
+		if (Config == EConfig::PerProfile)
+		{
+			const bool* ValuePtr = Profile->GetCustomBoolProperties().Find( KeyName );
+			return ValuePtr ? *ValuePtr : DefaultValue;
+
+		}
+		else
+		{
+			bool Value;
+			return GConfig->GetBool( Model->GetConfigSection() , *KeyName, Value, Model->GetConfigIni() ) ? Value : DefaultValue;
+		}
+	}
+
+
+	void FLaunchExtensionInstance::SetConfigString( EConfig Config, const TCHAR* Name, const FString& Value ) const
+	{
+		FString KeyName = GetConfigKeyName(Config, Name);
+
+		if (Config == EConfig::PerProfile)
+		{
+			if (Value.IsEmpty())
+			{
+				Profile->GetCustomStringProperties().Remove(KeyName);
+			}
+			else
+			{
+				Profile->GetCustomStringProperties().Add(KeyName, Value);
+			}
+			Marshaller->OnPropertyChanged();
+		}
+		else
+		{
+			GConfig->SetString( Model->GetConfigSection() , *KeyName, *Value, Model->GetConfigIni() );
+		}
+	}
+	void FLaunchExtensionInstance::SetConfigBool( EConfig Config, const TCHAR* Name, const bool& Value ) const
+	{
+		FString KeyName = GetConfigKeyName(Config, Name);
+
+		if (Config == EConfig::PerProfile)
+		{
+			if (Value == false)
+			{
+				Profile->GetCustomBoolProperties().Remove(KeyName);
+			}
+			else
+			{
+				Profile->GetCustomBoolProperties().Add(KeyName, Value);
+			}
+			Marshaller->OnPropertyChanged();
+		}
+		else
+		{
+			GConfig->SetBool( Model->GetConfigSection() , *KeyName, Value, Model->GetConfigIni() );
+		}
+	}
+
+	FString FLaunchExtensionInstance::GetConfigKeyName( EConfig Config, const TCHAR* Name ) const
+	{
+		switch (Config)
+		{
+			case EConfig::PerProfile:
+			case EConfig::User_Common:
+			{
+				return FString::Printf(TEXT("%s.%s"), Extension->GetInternalName(), Name );
+			}
+			break;
+
+			case EConfig::User_PerProfile:
+			{
+				FGuid ProfileId = Profile->GetId();
+				return FString::Printf(TEXT("%s.%s.%s"), Extension->GetInternalName(), *ProfileId.ToString(), Name );
+			}
+			break;
+		}
+
+		checkNoEntry();
+		return Name;
+	}
+}
+
+#undef LOCTEXT_NAMESPACE
diff --git a/Engine/Plugins/Developer/CustomLaunchUI/Source/CustomLaunchUI/Private/Marshalling/GenericProfileMarshaller.cpp b/Engine/Plugins/Developer/CustomLaunchUI/Source/CustomLaunchUI/Private/Marshalling/GenericProfileMarshaller.cpp
--- a/Engine/Plugins/Developer/CustomLaunchUI/Source/CustomLaunchUI/Private/Marshalling/GenericProfileMarshaller.cpp
+++ b/Engine/Plugins/Developer/CustomLaunchUI/Source/CustomLaunchUI/Private/Marshalling/GenericProfileMarshaller.cpp
@@ -77,11 +77,12 @@ namespace CustomLaunchUI
 
 	}
 
-	BEGIN_SLATE_FUNCTION_BUILD_OPTIMIZATION
 	void FGenericProfileMarshaller::Construct()
 	{
+		const TArray<FString>& DeviceIDs = Profile->GetDeployedDeviceGroup()->GetDeviceIDs();
+
 		MapOption = Profile->GetCookedMaps().Num() > 0 ? EMapOption::Selected : EMapOption::Startup;
-		DeployDeviceOption = Profile->GetDeployedDeviceGroup()->GetDeviceIDs().Num() > 0 ? EDeployDeviceOption::Selected : EDeployDeviceOption::Default;
+		DeployDeviceOption = DeviceIDs.Num() > 0 && !DeviceIDs[0].IsEmpty() ? EDeployDeviceOption::Selected : EDeployDeviceOption::Default;
 		bShouldCook = GetCook();
 		ContentScheme = Model->DetermineProfileContentScheme(Profile);
 		CachedBuildTargetType = Model->GetBuildTargetInfo(Profile).Type;
@@ -89,7 +90,6 @@ namespace CustomLaunchUI
 		CacheStartupMapList();
 		CacheArchitectures();
 	}
-	END_SLATE_FUNCTION_BUILD_OPTIMIZATION
 
 
 
@@ -103,6 +103,9 @@ namespace CustomLaunchUI
 	{
 		if (ProfileType == EProfileType::Basic)
 		{
+			// this is adding a new widget to the property tree
+			// - the first entry is the name of the property on the left-hand side
+			// - the seecond parameter is the widget itself that appears on the right-hand side
 			HeadingNode.AddWidget( LOCTEXT("ProjectLabel", "Project"), 
 				SNew(SCustomLaunchProjectCombo)
 				.OnSelectionChanged(this, &FGenericProfileMarshaller::SetProjectName)
@@ -130,6 +133,9 @@ namespace CustomLaunchUI
 	{
 		if (ProfileType == EProfileType::Basic)
 		{
+			// this is also adding a new widget to the property tree, as above.
+			// in this example the new struct parameter defines several callbacks that handle the 'reset to default' functionality. there are also options for disabling & hiding.
+			// the code is implemented with this syntax to avoid aid readability without using slate's TAttribute style functionality which seemed like an overkill for our simpler needs.
 			HeadingNode.AddWidget( LOCTEXT("TargetLabel", "Target"), 
 				{
 					.IsDefault = [this]()		{ return !Profile->HasBuildTargetSpecified() || Profile->GetBuildTarget().IsEmpty(); },
@@ -194,6 +200,8 @@ namespace CustomLaunchUI
 
 	void FGenericProfileMarshaller::AddCompressPakFilesProperty( FLaunchProfileTreeNode& HeadingNode )
 	{
+		// in this example we are adding a single boolean instead of a custom widget. the struct parameter defines how the value is accessed
+		// the ForPak is the lambda functions created in the constructor and is again aimed at readability
 		HeadingNode.AddBoolean( LOCTEXT("CompressPakFilesLabel", "Compress Pak Files"),
 			{
 				.GetValue = [this]()			{ return Profile->IsCompressed(); },
@@ -324,7 +332,7 @@ namespace CustomLaunchUI
 		);
 	}
 
-	void FGenericProfileMarshaller::AddBuidUATProperty( FLaunchProfileTreeNode& HeadingNode ) // fixme: not really relevant
+	void FGenericProfileMarshaller::AddBuidUATProperty( FLaunchProfileTreeNode& HeadingNode )
 	{
 		HeadingNode.AddBoolean( LOCTEXT("BuildUATLabel", "Build UAT"),
 			{
@@ -355,7 +363,7 @@ namespace CustomLaunchUI
 
 
 
-	void FGenericProfileMarshaller::AddStagingDirectoryProperty( FLaunchProfileTreeNode& HeadingNode ) // fixme: not really relevant
+	void FGenericProfileMarshaller::AddStagingDirectoryProperty( FLaunchProfileTreeNode& HeadingNode )
 	{
 		HeadingNode.AddDirectoryString( LOCTEXT("CustomStagingPathLabel", "Custom Stage Directory"),
 			{
@@ -407,7 +415,7 @@ namespace CustomLaunchUI
 		);
 	}
 
-	void FGenericProfileMarshaller::AddIncrementalDeployProperty( FLaunchProfileTreeNode& HeadingNode ) // fixme: not really relavent
+	void FGenericProfileMarshaller::AddIncrementalDeployProperty( FLaunchProfileTreeNode& HeadingNode )
 	{
 		HeadingNode.AddBoolean( LOCTEXT("DeployModifiedLabel", "Only Deploy Modified Content"),
 			{
@@ -467,7 +475,7 @@ namespace CustomLaunchUI
 
 	void FGenericProfileMarshaller::AddCommandLineProperty( FLaunchProfileTreeNode& HeadingNode )
 	{
-		HeadingNode.AddString( LOCTEXT("CommandLineLabel", "Command Line"), 
+		HeadingNode.AddCommandLineString( LOCTEXT("CommandLineLabel", "Additional Command Line"), 
 			{
 				.GetValue = [this]()				{ return GetCommandLine(); },
 				.SetValue = [this](FString Value)	{ SetCommandLine(Value); },
diff --git a/Engine/Plugins/Developer/CustomLaunchUI/Source/CustomLaunchUI/Private/Marshalling/ILaunchProfileMarshaller.cpp b/Engine/Plugins/Developer/CustomLaunchUI/Source/CustomLaunchUI/Private/Marshalling/ILaunchProfileMarshaller.cpp
--- a/Engine/Plugins/Developer/CustomLaunchUI/Source/CustomLaunchUI/Private/Marshalling/ILaunchProfileMarshaller.cpp
+++ b/Engine/Plugins/Developer/CustomLaunchUI/Source/CustomLaunchUI/Private/Marshalling/ILaunchProfileMarshaller.cpp
@@ -38,6 +38,11 @@ namespace CustomLaunchUI
 		{
 		}
 
+		virtual bool AllowExtensionsUI() const override
+		{
+			return false;
+		}
+
 	protected:
 		FLaunchProfileTreeDataRef TreeData;
 	};
@@ -86,6 +91,12 @@ namespace CustomLaunchUI
 				if (Marshaller.IsValid())
 				{
 					Marshaller->Construct();
+
+					if (Marshaller->AllowExtensionsUI())
+					{
+						Marshaller->GetProfileTree()->CreateExtensionsUI();
+					}
+
 					return Marshaller.ToSharedRef();
 				}
 			}
diff --git a/Engine/Plugins/Developer/CustomLaunchUI/Source/CustomLaunchUI/Private/Marshalling/LaunchProfileTreeData.cpp b/Engine/Plugins/Developer/CustomLaunchUI/Source/CustomLaunchUI/Private/Marshalling/LaunchProfileTreeData.cpp
--- a/Engine/Plugins/Developer/CustomLaunchUI/Source/CustomLaunchUI/Private/Marshalling/LaunchProfileTreeData.cpp
+++ b/Engine/Plugins/Developer/CustomLaunchUI/Source/CustomLaunchUI/Private/Marshalling/LaunchProfileTreeData.cpp
@@ -2,15 +2,21 @@
 
 #include "Marshalling/LaunchProfileTreeData.h"
 #include "Marshalling/ILaunchProfileMarshaller.h"
+#include "Extension/LaunchExtension.h"
 #include "Styling/CustomLaunchUIStyle.h"
+#include "Styling/AppStyle.h"
+#include "Styling/StyleColors.h"
 #include "Widgets/Input/SCheckBox.h"
 #include "Widgets/Input/SButton.h"
 #include "Widgets/Input/SEditableTextBox.h"
 #include "Widgets/Input/SMultiLineEditableTextBox.h"
+#include "Widgets/Input/SComboButton.h"
 #include "Widgets/Images/SImage.h"
 #include "Widgets/SBoxPanel.h"
 #include "Layout/Visibility.h"
 #include "Framework/Application/SlateApplication.h"
+#include "Framework/MultiBox/MultiBoxBuilder.h"
+#include "Framework/Commands/UIAction.h"
 #include "DesktopPlatformModule.h"
 
 
@@ -24,6 +30,11 @@ namespace CustomLaunchUI
 		, Marshaller(InMarshaller)
 	{
 		check(Marshaller);
+
+		if (InProfile.IsValid())
+		{
+			ExtensionInstances = FLaunchExtension::CreateExtensionInstancesForProfile(Profile.ToSharedRef(), InModel, InMarshaller);
+		}
 	}
 
 	FLaunchProfileTreeNode& FLaunchProfileTreeData::AddHeading( FText InName )
@@ -35,6 +46,14 @@ namespace CustomLaunchUI
 		return *TreeNode.Get();
 	}
 
+	void FLaunchProfileTreeData::CreateExtensionsUI()
+	{
+		for (TSharedPtr<FLaunchExtensionInstance> ExtensionInstance : ExtensionInstances)
+		{
+			ExtensionInstance->CustomizeTree(*this);
+		}
+	}
+
 
 
 	FLaunchProfileTreeNode::FLaunchProfileTreeNode( const FLaunchProfileTreeData* InTreeData ) 
@@ -249,6 +268,112 @@ namespace CustomLaunchUI
 	}
 
 
+	FLaunchProfileTreeNode& FLaunchProfileTreeNode::AddCommandLineString( FText InName, FStringCallbacks&& StringCallbacks )
+	{
+		check(StringCallbacks.GetValue);
+		check(StringCallbacks.SetValue);
+
+		FCallbacks WidgetCallbacks
+		{
+			.IsVisible = StringCallbacks.IsVisible,
+			.IsEnabled = StringCallbacks.IsEnabled,
+		};
+		if (StringCallbacks.GetDefaultValue != nullptr)
+		{
+			WidgetCallbacks.IsDefault = [StringCallbacks]()
+			{
+				return StringCallbacks.GetValue() == StringCallbacks.GetDefaultValue();
+			};
+
+			WidgetCallbacks.SetToDefault = [StringCallbacks]()
+			{
+				StringCallbacks.SetValue(StringCallbacks.GetDefaultValue());
+				// nb. OnPropertyChanged is called via SCustomLaunchWidgetTableRow's set to default
+			};
+		};
+
+		auto SetString = [this, StringCallbacks]( const FText& InText, ETextCommit::Type InTextCommit )
+		{
+			StringCallbacks.SetValue( InText.ToString() );
+			TreeData->Marshaller->OnPropertyChanged();
+		};
+
+		auto GetString = [StringCallbacks]()
+		{
+			return FText::FromString(StringCallbacks.GetValue());
+		};
+
+		auto OnGetMenuContent = [this,StringCallbacks]()
+		{
+			const bool bShouldCloseWindowAfterMenuSelection = false;
+			const bool bCloseSelfOnly = false;
+			const bool bSearchable = false;
+			FMenuBuilder MenuBuilder(bShouldCloseWindowAfterMenuSelection, nullptr, nullptr, bCloseSelfOnly, &FCoreStyle::Get(), bSearchable );
+
+			for (TSharedPtr<FLaunchExtensionInstance> ExtensionInstance : TreeData->ExtensionInstances)
+			{
+#if 1
+				// flat list of extensions - easier to see but risks getting cluttered later on
+				MenuBuilder.BeginSection(NAME_None, ExtensionInstance->GetExtension()->GetDisplayName());
+				ExtensionInstance->MakeCommandLineSubmenu(MenuBuilder);
+				MenuBuilder.EndSection();
+#else 
+				// extensions in submenus - more scalable, but harder to see everything at a glance
+				const bool bInOpenSubMenuOnClick = false;
+				MenuBuilder.AddSubMenu(
+					ExtensionInstance->GetExtension()->GetDisplayName()),
+					FText::GetEmpty(),
+					FNewMenuDelegate::CreateSP( ExtensionInstance.ToSharedRef(), &FLaunchExtensionInstance::MakeCommandLineSubmenu),
+					bInOpenSubMenuOnClick,
+					FSlateIcon(),
+					bShouldCloseWindowAfterMenuSelection
+				);
+#endif
+			}
+
+			return MenuBuilder.MakeWidget();
+		};
+
+
+		return AddWidget( InName,
+			MoveTemp(WidgetCallbacks),
+			SNew(SHorizontalBox)
+
+			// path field
+			+SHorizontalBox::Slot()
+			.FillWidth(1)
+			.VAlign(VAlign_Center)
+			[
+				SNew(SMultiLineEditableTextBox)
+				.AllowMultiLine(false)
+				.AutoWrapText(true)
+				.Font(FCoreStyle::Get().GetFontStyle("SmallFont"))
+				.Text_Lambda( GetString )
+				.OnTextCommitted_Lambda( SetString )
+			]
+
+			// command line extension button
+			+SHorizontalBox::Slot()
+			.AutoWidth()
+			.Padding(FMargin(0, 0, 4, 0))
+			.VAlign(VAlign_Center)
+			[
+				SNew(SComboButton)
+				.ComboButtonStyle(FAppStyle::Get(), "SimpleComboButton")
+				.ToolTipText(LOCTEXT("CmdLineOptionsLabel", "Add a special item to command line"))
+				.OnGetMenuContent_Lambda(OnGetMenuContent)
+				.Visibility_Lambda( [this]() { return TreeData->ExtensionInstances.Num() > 0 ? EVisibility::Visible : EVisibility::Collapsed; } )
+				.MenuPlacement(MenuPlacement_BelowRightAnchor)
+				.HasDownArrow(false)
+				.ButtonContent()
+				[
+					SNew(SImage)
+					.Image(FAppStyle::Get().GetBrush("Icons.Plus"))
+					.ColorAndOpacity(FStyleColors::AccentGreen)
+				]
+			]
+		);
+	}
 
 }
 
diff --git a/Engine/Plugins/Developer/CustomLaunchUI/Source/CustomLaunchUI/Private/Model/CustomLaunchUIModel.cpp b/Engine/Plugins/Developer/CustomLaunchUI/Source/CustomLaunchUI/Private/Model/CustomLaunchUIModel.cpp
--- a/Engine/Plugins/Developer/CustomLaunchUI/Source/CustomLaunchUI/Private/Model/CustomLaunchUIModel.cpp
+++ b/Engine/Plugins/Developer/CustomLaunchUI/Source/CustomLaunchUI/Private/Model/CustomLaunchUIModel.cpp
@@ -30,38 +30,49 @@ namespace CustomLaunchUI
 		, Launcher(InLauncher)
 		, ProfileManager(InProfileManager)
 	{
+		// register callbacks
+		ProfileManager->OnPostProcessLaunchCommandLine().AddRaw( this, &FModel::OnModifyLaunchCommandLine );
+
 		ProfileManager->OnProfileAdded().AddRaw(this, &FModel::HandleProfileManagerProfileAdded);
 		ProfileManager->OnProfileRemoved().AddRaw(this, &FModel::HandleProfileManagerProfileRemoved);
 
 		DeviceProxyManager->OnProxyAdded().AddRaw(this, &FModel::HandleDeviceProxyAdded);
 		DeviceProxyManager->OnProxyRemoved().AddRaw(this, &FModel::HandleDeviceProxyRemoved);
 
+
+		// ensure there's a project when we're in the editor
+		// this means we don't need to display the global project selector in the editor
 		if (GIsEditor && FPaths::IsProjectFilePathSet())
 		{
 			ProfileManager->SetProjectPath(FPaths::GetProjectFilePath());
 		}
 
-		LoadConfig();
 
-
-		AllProfiles = ProfileManager->GetAllProfiles();
-
-
-		// prepare basic launch profile
+		// prepare profiles
 		BasicLaunchProfile = CreateBasicLaunchProfile();
 		bHasSetBasicLaunchProfilePlatform = (BasicLaunchProfile->GetCookedPlatforms().Num() > 0);
 
+		AllProfiles = ProfileManager->GetAllProfiles();
 		AllProfiles.Add(BasicLaunchProfile);
 		SortProfiles();
 
 		DefaultBasicLaunchProfile = CreateBasicLaunchProfile();
 		DefaultCustomLaunchProfile = CreateCustomProfile(TEXT("DefaultCustomProfile"));
+
+
+		// use a custom ini file so that it can be shared between UnrealFrontend and the editor
+		ConfigFileName = FPaths::EngineSavedDir() / TEXT("Config") / TEXT("ProjectLauncher") / TEXT("UserSettings.ini");
+		FConfigContext::ReadSingleIntoGConfig().Load(*ConfigFileName);
+
+		LoadConfig();
 	}
 
 
 
 	FModel::~FModel()
 	{
+		ProfileManager->OnPostProcessLaunchCommandLine().RemoveAll( this );
+
 		ProfileManager->OnProfileAdded().RemoveAll(this);
 		ProfileManager->OnProfileRemoved().RemoveAll(this);
 
@@ -115,6 +126,17 @@ namespace CustomLaunchUI
 
 
 
+
+	const TCHAR* FModel::GetConfigSection() const
+	{
+		return TEXT("CustomLaunchUI"); 
+	}
+
+	const FString& FModel::GetConfigIni() const
+	{
+		return ConfigFileName;
+	}
+
 	void FModel::LoadConfig()
 	{
 		if (!GIsEditor)
@@ -132,7 +154,7 @@ namespace CustomLaunchUI
 			}
 			else if (GConfig != nullptr)
 			{
-				GConfig->GetString(TEXT("CustomLaunchUI"), TEXT("SelectedProjectPath"), ProjectPath, GEngineIni); // @todo: keeping the same
+				GConfig->GetString(GetConfigSection(), TEXT("SelectedProjectPath"), ProjectPath, GetConfigIni());
 			}
 	
 			ProfileManager->SetProjectPath(ProjectPath);
@@ -146,7 +168,7 @@ namespace CustomLaunchUI
 		if (!GIsEditor && GConfig != nullptr && !FPaths::IsProjectFilePathSet() && !FGameProjectHelper::IsGameAvailable(FApp::GetProjectName()))
 		{
 			FString ProjectPath = ProfileManager->GetProjectPath();
-			GConfig->SetString(TEXT("CustomLaunchUI"), TEXT("SelectedProjectPath"), *ProjectPath, GEngineIni);
+			GConfig->SetString(GetConfigSection(), TEXT("SelectedProjectPath"), *ProjectPath, GetConfigIni());
 		}
 	}
 
@@ -646,6 +668,23 @@ namespace CustomLaunchUI
 
 
 
+	extern void ApplyExtensionVariables( const ILauncherProfileRef& InProfile, FString& InOutCommandLine, TSharedRef<FModel> InModel );
+
+	void FModel::OnModifyLaunchCommandLine( const ILauncherProfileRef& InProfile, FString& InOutCommandLine )
+	{
+		ApplyExtensionVariables(InProfile, InOutCommandLine, AsShared() );
+	}
+
+
+
+
+
+
+
+
+
+
+
 	TArray<EContentScheme> GetAllContentSchemes()
 	{
 		TArray<EContentScheme> ContentSchemes;
diff --git a/Engine/Plugins/Developer/CustomLaunchUI/Source/CustomLaunchUI/Private/Widgets/Output/SCustomLaunchLaunchPanel.cpp b/Engine/Plugins/Developer/CustomLaunchUI/Source/CustomLaunchUI/Private/Widgets/Output/SCustomLaunchLaunchPanel.cpp
--- a/Engine/Plugins/Developer/CustomLaunchUI/Source/CustomLaunchUI/Private/Widgets/Output/SCustomLaunchLaunchPanel.cpp
+++ b/Engine/Plugins/Developer/CustomLaunchUI/Source/CustomLaunchUI/Private/Widgets/Output/SCustomLaunchLaunchPanel.cpp
@@ -357,12 +357,16 @@ void SCustomLaunchLaunchPanel::SetLauncherWorker( const ILauncherWorkerRef& Work
 		}
 	}
 
+	OutputLog->RefreshLog();
+
+	Worker->OnOutputReceived().AddRaw(this, &SCustomLaunchLaunchPanel::HandleOutputReceived);
+}
 
+void SCustomLaunchLaunchPanel::ClearLog()
+{
 	Model->ClearLogMessages();
 	LaunchLogTextMarshaller->RefreshAllLogMessages();
 	OutputLog->RefreshLog();
-
-	Worker->OnOutputReceived().AddRaw(this, &SCustomLaunchLaunchPanel::HandleOutputReceived);
 }
 
 
diff --git a/Engine/Plugins/Developer/CustomLaunchUI/Source/CustomLaunchUI/Private/Widgets/Output/SCustomLaunchLaunchPanel.h b/Engine/Plugins/Developer/CustomLaunchUI/Source/CustomLaunchUI/Private/Widgets/Output/SCustomLaunchLaunchPanel.h
--- a/Engine/Plugins/Developer/CustomLaunchUI/Source/CustomLaunchUI/Private/Widgets/Output/SCustomLaunchLaunchPanel.h
+++ b/Engine/Plugins/Developer/CustomLaunchUI/Source/CustomLaunchUI/Private/Widgets/Output/SCustomLaunchLaunchPanel.h
@@ -32,6 +32,7 @@ public:
 
 
 	void SetLauncherWorker( const ILauncherWorkerRef& Worker );
+	void ClearLog();
 
 private:
 	bool IsBuilding() const;
diff --git a/Engine/Plugins/Developer/CustomLaunchUI/Source/CustomLaunchUI/Private/Widgets/SCustomLaunchUI.cpp b/Engine/Plugins/Developer/CustomLaunchUI/Source/CustomLaunchUI/Private/Widgets/SCustomLaunchUI.cpp
--- a/Engine/Plugins/Developer/CustomLaunchUI/Source/CustomLaunchUI/Private/Widgets/SCustomLaunchUI.cpp
+++ b/Engine/Plugins/Developer/CustomLaunchUI/Source/CustomLaunchUI/Private/Widgets/SCustomLaunchUI.cpp
@@ -64,6 +64,8 @@ END_SLATE_FUNCTION_BUILD_OPTIMIZATION
 
 void SCustomLaunchUI::OnProfileRun(const ILauncherProfilePtr& Profile)
 {
+	LaunchProgressPanel->ClearLog();
+
 	LauncherProfile = Profile;
 	if (LauncherProfile.IsValid())
 	{
@@ -97,6 +99,8 @@ FReply SCustomLaunchUI::OnProgressClose()
 
 FReply SCustomLaunchUI::OnRerunClicked()
 {
+	LaunchProgressPanel->ClearLog();
+
 	if (LauncherWorker.IsValid())
 	{
 		LauncherWorker->Cancel();
diff --git a/Engine/Plugins/Developer/CustomLaunchUI/Source/CustomLaunchUI/Public/CustomLaunchUIModule.h b/Engine/Plugins/Developer/CustomLaunchUI/Source/CustomLaunchUI/Public/CustomLaunchUIModule.h
--- a/Engine/Plugins/Developer/CustomLaunchUI/Source/CustomLaunchUI/Public/CustomLaunchUIModule.h
+++ b/Engine/Plugins/Developer/CustomLaunchUI/Source/CustomLaunchUI/Public/CustomLaunchUIModule.h
@@ -8,6 +8,7 @@
 namespace CustomLaunchUI
 {
 	class ILaunchProfileMarshallerFactory;
+	class FLaunchExtension;
 };
 
 /*
@@ -22,10 +23,20 @@ public:
 	*/
 	static inline ICustomLaunchUIModule& Get()
 	{
-		return FModuleManager::LoadModuleChecked<ICustomLaunchUIModule>("CustomLaunchUIModule");
+		return FModuleManager::LoadModuleChecked<ICustomLaunchUIModule>("CustomLaunchUI");
+	}
+
+	/**
+	* Singleton-like access to the Custom Launch UI module instance
+	* @return Returns ICustomLaunchUIModule singleton instance if it is loaded
+	*/
+	static inline ICustomLaunchUIModule* TryGet()
+	{
+		return FModuleManager::GetModulePtr<ICustomLaunchUIModule>("CustomLaunchUI");
 	}
 
 
+
 	/**
 	 * Registers a marhaller factory, for defining the layout & fields for editing an ILauncherProfile
 	 */
@@ -35,5 +46,17 @@ public:
 	 * Unregisters a previously-registered marshaller factory
 	 */
 	virtual void UnregisterMarshaller( TSharedRef<CustomLaunchUI::ILaunchProfileMarshallerFactory> MarshallerFactory ) = 0;
+
+
+
+	/**
+	 * Registers an extension
+	 */
+	virtual void RegisterExtension( TSharedRef<CustomLaunchUI::FLaunchExtension> Extension ) = 0;
+
+	/**
+	 * Unregisters a previously-registered extension
+	 */
+	virtual void UnregisterExtension( TSharedRef<CustomLaunchUI::FLaunchExtension> Extension ) = 0;
 };
 
diff --git a/Engine/Plugins/Developer/CustomLaunchUI/Source/CustomLaunchUI/Public/Extension/LaunchExtension.h b/Engine/Plugins/Developer/CustomLaunchUI/Source/CustomLaunchUI/Public/Extension/LaunchExtension.h
new file mode 100644
--- /dev/null
+++ b/Engine/Plugins/Developer/CustomLaunchUI/Source/CustomLaunchUI/Public/Extension/LaunchExtension.h
@@ -0,0 +1,285 @@
+// Copyright Epic Games, Inc. All Rights Reserved.
+
+#pragma once
+
+#include "Internationalization/Text.h"
+#include "Marshalling/LaunchProfileTreeData.h"
+#include "ILauncherProfile.h"
+#include "Framework/MultiBox/MultiBoxBuilder.h"
+
+namespace CustomLaunchUI
+{
+	class FLaunchExtension;
+
+	/**
+	 * Base class for a launch extension instance.
+	 * Used while editing a specific profile, and when finalizing the command line arguments during profile launch.
+	 * 
+	 * Created by a specialization of FLaunchExtension as follows:
+	 * 
+	 * TSharedPtr<CustomLaunchUI:FLaunchExtensinInstance> FMyLaunchExtension::CreateInstanceForProfile( CustomLaunchUI::FLaunchExtensionInstance::FArgs& InArgs )
+	 * {
+	 *     return MakeShared<FMyLaunchExtensionInstance>(InArgs);
+	 * }
+	 */
+	class CUSTOMLAUNCHUI_API FLaunchExtensionInstance : public TSharedFromThis<FLaunchExtensionInstance>
+	{
+	public:
+		struct FArgs
+		{
+			ILauncherProfileRef Profile;
+			ILaunchProfileMarshaller* Marshaller;
+			TSharedRef<FModel> Model;
+			TSharedRef<FLaunchExtension> Extension;
+		};
+
+		FLaunchExtensionInstance( FArgs& InArgs );
+		virtual ~FLaunchExtensionInstance();
+
+		/**
+		 * Returns the parameters that this extension provides. They will be added to the submenu.
+		 */
+		virtual bool GetExtensionParameters( TArray<FString>& OutParameters ) const;
+
+		/**
+		 * Returns the user-facing name for the given parameter. It will default to the parameter itself.
+		 */
+		virtual FText GetExtensionParameterDisplayName( const FString& InParameter ) const;
+
+		/**
+		 * Returns the user facing variables that this extension provides, in "$(name)" format
+		 */
+		virtual bool GetExtensionVariables( TArray<FString>& OutVariables ) const;
+
+		/**
+		 * Returns the current value for the given variable
+		 */
+		virtual bool GetExtensionVariableValue( const FString& InParameter, FString& OutValue ) const;
+
+		/**
+		 * Hook to allow the extension to extend its menu
+		 */
+		virtual void CustomizeSubmenu( FMenuBuilder& MenuBuilder ) {};
+
+		/**
+		 * Hook to allow the extension to add extra fields to the property editing tree, if the marshall allows it
+		 * Property tree items should be hidden until the user has selected something to make it relevent, to avoid cluttering the UI.
+		 */
+		virtual void CustomizeTree( FLaunchProfileTreeData& ProfileTreeData ) {};
+
+		/**
+		 * Advanced hook to allow any advanced modification of the command line when our profile is launched
+		 */
+		virtual void CustomizeLaunchCommandLine( FString& InOutCommandLine ) {};
+
+	protected:
+
+		/**
+		 * Determine if the given parameter is on the command line
+ 		 * 
+		 * @param InParameter the parameter to test, e.g. -key=value or -param
+		 * @returns whether the parameter is being used
+		 */
+		bool IsParameterUsed( const FString& InParameter ) const;
+
+		/** 
+		 * Add or remove the given parameter on the command line
+		 * 
+		 * @param InParameter the parameter to add or remove, e.g. -key=value or -param
+		 * @param bUsed whether to add or remove the parameter
+		 */
+		void SetParameterUsed( const FString& InParameter, bool bUsed );
+
+		/** 
+		 * Add the given parameter to the command line
+		 * 
+		 * @param InParameter the parameter to add, e.g. -key=value or -param
+		 */
+		void AddParameter( const FString& InParameter );
+
+		/** 
+		 * Remove the given parameter from the command line
+		 * 
+		 * @param InParameter the parameter to remove, e.g. -key=value or -param
+		 */
+		void RemoveParameter( const FString& InParameter );
+
+		/** 
+		 * Retrieve the current value for the given command line parameter
+		 * 
+		 * @param InParameter the parameter to query, e.g. -key=value
+		 * @returns the value of the parameter
+		 */
+		FString GetParameterValue( const FString& InParameter ) const;
+
+		/** 
+		 * Update the value of the given command line parameter
+		 * 
+		 * @param InParameter the parameter to modify, e.g. -key=value or -key=
+		 * @param InNewValue the new value to use
+		 * @returns true unless InParameter is not a key/value property
+		 */
+		bool UpdateParameterValue( const FString& InParameter, const FString& InNewValue );
+
+		/** 
+		 * Get the current state of the given command line parameter, allowing for user-changed values
+		 * 
+		 * @param InParameter the parameter to test, e.g. -key=value
+		 * @returns The parameter as it is now, allowing for user modifications, e.g. -key=value or -key=some_new_value
+		 */
+		FString GetFinalParameter( const FString& InParameter ) const;
+
+		/** 
+		 * Get the current command line
+		 * Extensions should generally go through this function instead of querying the Profile directlry
+		 *
+		 * @returns the current command line string
+		 */
+		FString GetCommandLine() const;
+
+		/** 
+		 * Update the current command line
+		 * Extensions should generally go through this function instead of modifying the Profile directlry
+		 * 
+		 * @param CommandLine the new command line
+		 */
+		void SetCommandLine( const FString& CommandLine );
+
+		/** 
+		 * Get the profile we were instantiated for
+		 */
+		inline ILauncherProfileRef GetProfile() const { return Profile; }
+
+		/** 
+		 * Get the model, for general purpose helper functions
+		 */
+		inline TSharedRef<FModel> GetModel() const { return Model; }
+
+		/** 
+		 * Get the extension that instantiated us
+		 */
+		inline TSharedRef<FLaunchExtension> GetExtension() const { return Extension; }
+
+
+		/** 
+		* Enumeration to choose where a value should be stored
+		*/
+		enum class EConfig : uint8
+		{
+			User_Common,     ///< value is shared between all instances of this extension
+			User_PerProfile, ///< value is specific to this profile & extension
+			PerProfile,      ///< value is saved with the profile
+		};
+
+		/**
+		 * Read a configuration string value
+		 * 
+		 * @param Config where the value is stored
+		 * @param Name name of the value to read
+		 * @param DefaultValue value to return if the value isn't found
+		 * @returns the value of the string, or DefaultValue if it isn't found
+		 */
+		FString GetConfigString( EConfig Config, const TCHAR* Name, const TCHAR* DefaultValue = TEXT("") ) const;
+
+		/**
+		 * Read a configuration bool value
+		 * 
+		 * @param Config where the value is stored
+		 * @param Name name of the value to read
+		 * @param DefaultValue value to return if the value isn't found
+		 * @returns the value of the bool, or DefaultValue if it isn't found
+		 */
+		bool GetConfigBool( EConfig Config, const TCHAR* Name, bool DefaultValue = false ) const;
+
+		/**
+		 * Write a configuration string value
+		 * 
+		 * @param Config where the value is stored
+		 * @param Name name of the value to write
+		 * @param Value value to write
+		 */
+		void SetConfigString( EConfig Config, const TCHAR* Name, const FString& Value ) const;
+
+		/**
+		 * Write a configuration bool value
+		 * 
+		 * @param Config where the value is stored
+		 * @param Name name of the value to write
+		 * @param Value value to write
+		 */
+		void SetConfigBool( EConfig Config, const TCHAR* Name, const bool& Value ) const;
+
+		/**
+		 * Get the final key name to use for reading & writing a configuration value
+		 * 
+		 * @param Config where the value is stored
+		 * @param Name name of the value
+		 * 
+		 * @returns namespaced version of Name
+		 */
+		FString GetConfigKeyName( EConfig Config, const TCHAR* Name ) const;
+
+
+
+	private:
+		bool IsParameterGroup( const FString& InParameter ) const;
+		bool TryRemoveParameterGroup( const FString& InParameter );
+
+
+		ILauncherProfileRef Profile;
+		ILaunchProfileMarshaller* Marshaller;
+		TSharedRef<FModel> Model;
+		TSharedRef<FLaunchExtension> Extension;
+
+		friend class FLaunchProfileTreeNode;
+		void MakeCommandLineSubmenu( FMenuBuilder& MenuBuilder );
+
+	};
+
+
+
+	/**
+	 * Base class for a launch extension.
+	 * 
+	 * Singleton instance is registered with this plugin during initialization as follows:
+	 * 
+	 *   TSharedPtr<FMyLaunchExtension> MyExtension = MakeShared<FMyLaunchExtension>();
+	 *   ICustomLaunchUIModule::Get().RegisterExtension( MyExtension.ToSharedRef() );
+	 */
+	class CUSTOMLAUNCHUI_API FLaunchExtension : public TSharedFromThis<FLaunchExtension>
+	{
+	public:
+		virtual ~FLaunchExtension() = default;
+
+		/**
+		 * Create an instance of the launch extension for the given profile
+		 *
+		 * @returns: new instance, or null if it isn't appropriate
+		 */
+		virtual TSharedPtr<FLaunchExtensionInstance> CreateInstanceForProfile( FLaunchExtensionInstance::FArgs& InArgs ) = 0;
+
+		/**
+		 * Returns the debug name for this extension
+		 */
+		virtual const TCHAR* GetInternalName() const = 0;
+
+		/**
+		 * Returns the user-facing name for this extension
+		 */
+		virtual FText GetDisplayName() const = 0;
+
+
+		/**
+		 * Instantiate all compatible extensions
+		 * 
+		 * @param InProfile the current profile 
+		 * @param InMarshaller the active marshaller
+		 * @param InModel helper class
+		 * @returns array of all compatible extensions
+		 */
+		static TArray<TSharedPtr<FLaunchExtensionInstance>> CreateExtensionInstancesForProfile( ILauncherProfileRef InProfile, TSharedRef<FModel> InModel, ILaunchProfileMarshaller* InMarshaller );
+	};
+
+
+};
+
diff --git a/Engine/Plugins/Developer/CustomLaunchUI/Source/CustomLaunchUI/Public/Marshalling/GenericProfileMarshaller.h b/Engine/Plugins/Developer/CustomLaunchUI/Source/CustomLaunchUI/Public/Marshalling/GenericProfileMarshaller.h
--- a/Engine/Plugins/Developer/CustomLaunchUI/Source/CustomLaunchUI/Public/Marshalling/GenericProfileMarshaller.h
+++ b/Engine/Plugins/Developer/CustomLaunchUI/Source/CustomLaunchUI/Public/Marshalling/GenericProfileMarshaller.h
@@ -12,7 +12,16 @@ namespace ESelectInfo { enum Type : int; }
 
 namespace CustomLaunchUI
 {
-
+	/**
+	 * Base class for a profile marshaller that creates FLaunchProfileTreeData from a given ILauncherProfile.
+	 * 
+	 * Expected to be created by an instance of ILaunchProfileMarshallerFactory, for example:
+	 * 
+	 *    TSharedPtr<ILaunchProfileMarshaller> FMyProfileMarshallerFactory::TryCreateMarshaller( const ILauncherProfileRef& InProfile, const TSharedRef<FModel>& InModel )
+	 *    {
+	 *        return MakeShared<FMyProfileMarshaller>(InProfile, InModel);
+	 *    }
+	 */
 	class CUSTOMLAUNCHUI_API FGenericProfileMarshaller : public ILaunchProfileMarshaller, public TSharedFromThis<FGenericProfileMarshaller>
 	{
 	public:
@@ -26,6 +35,11 @@ namespace CustomLaunchUI
 			return TreeData;
 		}
 
+		virtual bool AllowExtensionsUI() const override
+		{
+			return true;
+		}
+
 	protected:
 		// default property creation functions
 		void AddProjectProperty( FLaunchProfileTreeNode& HeadingNode );
@@ -44,13 +58,13 @@ namespace CustomLaunchUI
 		void AddAdditionalCookerOptionsProperty( FLaunchProfileTreeNode& HeadingNode );
 		void AddBuildProperty( FLaunchProfileTreeNode& HeadingNode );
 		void AddForceBuildProperty( FLaunchProfileTreeNode& HeadingNode );
-		void AddBuidUATProperty( FLaunchProfileTreeNode& HeadingNode ); // fixme: not really relevant
+		void AddBuidUATProperty( FLaunchProfileTreeNode& HeadingNode );
 		void AddArchitectureProperty( FLaunchProfileTreeNode& HeadingNode );
-		void AddStagingDirectoryProperty( FLaunchProfileTreeNode& HeadingNode ); // fixme: not really relevant
+		void AddStagingDirectoryProperty( FLaunchProfileTreeNode& HeadingNode );
 		void AddArchiveBuildProperty( FLaunchProfileTreeNode& HeadingNode );
 		void AddArchiveBuildDirectoryProperty( FLaunchProfileTreeNode& HeadingNode );
 		void AddDeployProperty( FLaunchProfileTreeNode& HeadingNode );
-		void AddIncrementalDeployProperty( FLaunchProfileTreeNode& HeadingNode ); // fixme: not really relavent
+		void AddIncrementalDeployProperty( FLaunchProfileTreeNode& HeadingNode );
 		void AddTargetDeviceProperty( FLaunchProfileTreeNode& HeadingNode );
 		void AddRunProperty( FLaunchProfileTreeNode& HeadingNode );
 		void AddInitialMapProperty( FLaunchProfileTreeNode& HeadingNode );
diff --git a/Engine/Plugins/Developer/CustomLaunchUI/Source/CustomLaunchUI/Public/Marshalling/ILaunchProfileMarshaller.h b/Engine/Plugins/Developer/CustomLaunchUI/Source/CustomLaunchUI/Public/Marshalling/ILaunchProfileMarshaller.h
--- a/Engine/Plugins/Developer/CustomLaunchUI/Source/CustomLaunchUI/Public/Marshalling/ILaunchProfileMarshaller.h
+++ b/Engine/Plugins/Developer/CustomLaunchUI/Source/CustomLaunchUI/Public/Marshalling/ILaunchProfileMarshaller.h
@@ -11,7 +11,13 @@ namespace CustomLaunchUI
 
 	/**
 	 * Interface for a profile marshaller that creates FLaunchProfileTreeData from a given ILauncherProfile.
-	 * Expected to be created by an instance of ILaunchProfileMarshallerFactory
+	 * 
+	 * Expected to be created by an instance of ILaunchProfileMarshallerFactory, for example:
+	 * 
+	 *    TSharedPtr<ILaunchProfileMarshaller> FMyProfileMarshallerFactory::TryCreateMarshaller( const ILauncherProfileRef& InProfile, const TSharedRef<FModel>& InModel )
+	 *    {
+	 *        return MakeShared<FMyProfileMarshaller>(InProfile, InModel);
+	 *    }
 	 */
 	class ILaunchProfileMarshaller
 	{
@@ -40,11 +46,21 @@ namespace CustomLaunchUI
 		 * NOTE: Custom widgets will need to call this manually.
 		 */
 		virtual void OnPropertyChanged() = 0;
+
+		/**
+		 * Whether this marshaller allows extensions to add UI elements. Typically would be true without good reason.
+		 */
+		virtual bool AllowExtensionsUI() const = 0;
 	};
 
 
 	/**
 	 * Interface for a factory that can create specializations of ILaunchProfileMarshaller for a given ILauncherProfile
+	 * 
+	 * Singleton instance is registered with this plugin during initialization as follows:
+	 * 
+	 *   TSharedPtr<FMyMarshallerFactory> MyMarshallar = MakeShared<FMyMarshallerFactory>();
+	 *   ICustomLaunchUIModule::Get().RegisterMarshaller( MyMarshaller.ToSharedRef() );
 	 */
 	class ILaunchProfileMarshallerFactory
 	{
diff --git a/Engine/Plugins/Developer/CustomLaunchUI/Source/CustomLaunchUI/Public/Marshalling/LaunchProfileTreeData.h b/Engine/Plugins/Developer/CustomLaunchUI/Source/CustomLaunchUI/Public/Marshalling/LaunchProfileTreeData.h
--- a/Engine/Plugins/Developer/CustomLaunchUI/Source/CustomLaunchUI/Public/Marshalling/LaunchProfileTreeData.h
+++ b/Engine/Plugins/Developer/CustomLaunchUI/Source/CustomLaunchUI/Public/Marshalling/LaunchProfileTreeData.h
@@ -15,6 +15,7 @@ namespace CustomLaunchUI
 	class ILaunchProfileMarshaller;
 	class FLaunchProfileTreeData;
 	class FLaunchProfileTreeNode;
+	class FLaunchExtensionInstance;
 	class FModel;
 
 	typedef TSharedPtr<FLaunchProfileTreeData> FLaunchProfileTreeDataPtr;
@@ -69,6 +70,7 @@ namespace CustomLaunchUI
 		};
 		FLaunchProfileTreeNode& AddString( FText InName, FStringCallbacks&& StringCallbacks );
 		FLaunchProfileTreeNode& AddDirectoryString( FText InName, FStringCallbacks&& StringCallbacks );
+		FLaunchProfileTreeNode& AddCommandLineString( FText InName, FStringCallbacks&& StringCallbacks );
 
 
 
@@ -92,16 +94,18 @@ namespace CustomLaunchUI
 	{
 	public:
 		// construction
-		FLaunchProfileTreeData(ILauncherProfilePtr InProfile, TSharedRef<CustomLaunchUI::FModel> InModel, ILaunchProfileMarshaller* InMarshaller);
+		FLaunchProfileTreeData(ILauncherProfilePtr InProfile, TSharedRef<FModel> InModel, ILaunchProfileMarshaller* InMarshaller);
 
 		// node creation
 		FLaunchProfileTreeNode& AddHeading( FText InName );
+		void CreateExtensionsUI();
 
 		// member variables
 		ILauncherProfilePtr Profile;
-		TSharedRef<CustomLaunchUI::FModel> Model;
+		TSharedRef<FModel> Model;
 		TArray<FLaunchProfileTreeNodePtr> Nodes;
 		ILaunchProfileMarshaller* Marshaller;
+		TArray<TSharedPtr<FLaunchExtensionInstance>> ExtensionInstances;
 	};
 
 
diff --git a/Engine/Plugins/Developer/CustomLaunchUI/Source/CustomLaunchUI/Public/Model/CustomLaunchUIModel.h b/Engine/Plugins/Developer/CustomLaunchUI/Source/CustomLaunchUI/Public/Model/CustomLaunchUIModel.h
--- a/Engine/Plugins/Developer/CustomLaunchUI/Source/CustomLaunchUI/Public/Model/CustomLaunchUIModel.h
+++ b/Engine/Plugins/Developer/CustomLaunchUI/Source/CustomLaunchUI/Public/Model/CustomLaunchUIModel.h
@@ -84,8 +84,7 @@ namespace CustomLaunchUI
 
 	extern bool bUseFriendlyBuildTargetSelection;
 
-
-	class FModel : public TSharedFromThis<FModel>
+	class CUSTOMLAUNCHUI_API FModel : public TSharedFromThis<FModel>
 	{
 	public:
 		FModel(const TSharedRef<ITargetDeviceProxyManager>& InDeviceProxyManager, const TSharedRef<ILauncher>& InLauncher, const TSharedRef<ILauncherProfileManager>& InProfileManager);
@@ -145,7 +144,7 @@ namespace CustomLaunchUI
 		EContentScheme DetermineProfileContentScheme(const ILauncherProfileRef& Profile) const;
 		void SetProfileContentScheme(EContentScheme ContentScheme, const ILauncherProfileRef& Profile, bool bWantToCook = true, ELauncherProfileDeploymentModes::Type DefaultDeploymentMode = ELauncherProfileDeploymentModes::CopyToDevice);
 
-		TSharedPtr<FLaunchLogMessage> AddLogMessage( const FString& InMessage, ELogVerbosity::Type InVerbosity );
+		TSharedPtr<FLaunchLogMessage> AddLogMessage( const FString& InMessage, ELogVerbosity::Type InVerbosity = ELogVerbosity::Log );
 		void ClearLogMessages();
 		int32 GetNumLogMessages() const { return LaunchLogMessages.Num(); }
 
@@ -169,7 +168,12 @@ namespace CustomLaunchUI
 		const FProjectSettings GetProjectSettings( const FString& ProjectPath );
 		const FProjectSettings GetProjectSettings( const ILauncherProfileRef& Profile );
 
+		const TCHAR* GetConfigSection() const;
+		const FString& GetConfigIni() const;
+
 	private:
+		FString ConfigFileName;
+
 		void LoadConfig();
 		void SaveConfig();
 
@@ -180,6 +184,8 @@ namespace CustomLaunchUI
 
 		void ReadProjectSettingsFromConfig( FConfigCacheIni& InConfig, FProjectSettings& OutResult );
 
+		void OnModifyLaunchCommandLine( const ILauncherProfileRef& InProfile, FString& InOutCommandLine );
+
 		TSharedRef<ITargetDeviceProxyManager> DeviceProxyManager;
 		TSharedRef<ILauncher> Launcher;
 		TSharedRef<ILauncherProfileManager> ProfileManager;
