BuildCookRunUI tweaks & fixes:
- Launch button icon changes to reflect what will happen.
- Extra-long profile name & descriptions will be correctly cropped.
- Add 'Edit Description' option to profile's right-click menu.
- Add Hint Text to the Description so it's clear what the field is for.
- Don't auto-scroll the output log if there is any text selected.

#rnx

diff --git a/Engine/Plugins/Developer/CustomLaunchUI/Source/CustomLaunchUI/Private/Widgets/Output/SCustomLaunchOutputLog.cpp b/Engine/Plugins/Developer/CustomLaunchUI/Source/CustomLaunchUI/Private/Widgets/Output/SCustomLaunchOutputLog.cpp
--- a/Engine/Plugins/Developer/CustomLaunchUI/Source/CustomLaunchUI/Private/Widgets/Output/SCustomLaunchOutputLog.cpp
+++ b/Engine/Plugins/Developer/CustomLaunchUI/Source/CustomLaunchUI/Private/Widgets/Output/SCustomLaunchOutputLog.cpp
@@ -377,6 +377,7 @@ void SCustomLaunchOutputLog::RefreshLog()
 void SCustomLaunchOutputLog::RequestForceScroll(bool bIfUserHasNotScrolledUp)
 {
 	if (LaunchLogTextMarshaller->GetNumFilteredMessages() > 0
+		&& !LogMessageTextBox->AnyTextSelected()
 		&& (!bIfUserHasNotScrolledUp || !bIsUserScrolled))
 	{
 		LogMessageTextBox->ScrollTo(ETextLocation::EndOfDocument);
diff --git a/Engine/Plugins/Developer/CustomLaunchUI/Source/CustomLaunchUI/Private/Widgets/Profiles/SCustomLaunchCustomProfileSelector.cpp b/Engine/Plugins/Developer/CustomLaunchUI/Source/CustomLaunchUI/Private/Widgets/Profiles/SCustomLaunchCustomProfileSelector.cpp
--- a/Engine/Plugins/Developer/CustomLaunchUI/Source/CustomLaunchUI/Private/Widgets/Profiles/SCustomLaunchCustomProfileSelector.cpp
+++ b/Engine/Plugins/Developer/CustomLaunchUI/Source/CustomLaunchUI/Private/Widgets/Profiles/SCustomLaunchCustomProfileSelector.cpp
@@ -62,6 +62,7 @@ void SCustomLaunchCustomProfileSelector::Construct(const FArguments& InArgs, con
 	OnProfileDuplicate = InArgs._OnProfileDuplicate;
 	OnProfileEdit = InArgs._OnProfileEdit;
 	OnProfileRename = InArgs._OnProfileRename;
+	OnProfileEditDescription = InArgs._OnProfileEditDescription;
 	OnProfileModified = InArgs._OnProfileModified;
 	ChangeProfileEditorVisibility = InArgs._ChangeProfileEditorVisibility;
 	EditPanelVisible = InArgs._EditPanelVisible;
@@ -231,6 +232,7 @@ TSharedRef<ITableRow> SCustomLaunchCustomProfileSelector::GenerateCustomProfileR
 			.OnTextCommitted( this, &SCustomLaunchCustomProfileSelector::SetProfileDescription, Profile )
 			.OnExitEditingMode_Lambda( [this]() { Model->SortProfiles(); RefreshCustomProfileList(); } )
 			.ColorAndOpacity(FSlateColor::UseSubduedForeground())
+			.HintText(LOCTEXT("DescriptionHintText", "Enter a description for this profile"))
 		];
 
 		NameEditTextBoxes.Add(Profile, NameEditTextBox);
@@ -409,6 +411,15 @@ TSharedPtr<SWidget> SCustomLaunchCustomProfileSelector::MakeContextMenu()
 			),
 			NAME_None);
 
+		MenuBuilder.AddMenuEntry(
+			LOCTEXT("CustomProfileEditDescLabel", "Edit Description"),
+			LOCTEXT("CustomProfileEditDescToolTip", "Edit the description of this profile"),
+			FSlateIcon(FAppStyle::Get().GetStyleSetName(), "Icons.Edit"),
+			FUIAction(
+				FExecuteAction::CreateSP(this, &SCustomLaunchCustomProfileSelector::OnCustomProfileEditDescriptionClicked, Profile)
+			),
+			NAME_None);
+
 		MenuBuilder.AddMenuEntry(
 			LOCTEXT("CustomProfileDuplicateLabel", "Duplicate Profile"),
 			LOCTEXT("CustomProfileDuplicateToolTip", "Duplicate this profile"),
@@ -461,6 +472,11 @@ void SCustomLaunchCustomProfileSelector::OnCustomProfileRenameClicked(ILauncherP
 	OnProfileRename.ExecuteIfBound(Profile);
 }
 
+void SCustomLaunchCustomProfileSelector::OnCustomProfileEditDescriptionClicked(ILauncherProfilePtr Profile)
+{
+	OnProfileEditDescription.ExecuteIfBound(Profile);
+}
+
 void SCustomLaunchCustomProfileSelector::OnDeviceRemoved( const FString DeviceID, ILauncherProfilePtr Profile)
 {
 	if (ensure(Profile->GetDeployedDeviceGroup() != nullptr) && Profile->GetDeployedDeviceGroup()->GetDeviceIDs().Contains(DeviceID))
diff --git a/Engine/Plugins/Developer/CustomLaunchUI/Source/CustomLaunchUI/Private/Widgets/Profiles/SCustomLaunchCustomProfileSelector.h b/Engine/Plugins/Developer/CustomLaunchUI/Source/CustomLaunchUI/Private/Widgets/Profiles/SCustomLaunchCustomProfileSelector.h
--- a/Engine/Plugins/Developer/CustomLaunchUI/Source/CustomLaunchUI/Private/Widgets/Profiles/SCustomLaunchCustomProfileSelector.h
+++ b/Engine/Plugins/Developer/CustomLaunchUI/Source/CustomLaunchUI/Private/Widgets/Profiles/SCustomLaunchCustomProfileSelector.h
@@ -27,6 +27,7 @@ public:
 		SLATE_EVENT(CustomLaunchUI::FOnProfileClicked, OnProfileDuplicate)
 		SLATE_EVENT(CustomLaunchUI::FOnProfileClicked, OnProfileDelete)
 		SLATE_EVENT(CustomLaunchUI::FOnProfileClicked, OnProfileRename)
+		SLATE_EVENT(CustomLaunchUI::FOnProfileClicked, OnProfileEditDescription)
 		SLATE_EVENT(CustomLaunchUI::FOnProfileClicked, OnProfileModified)
 		SLATE_EVENT(FSetProfileEditorVisible, ChangeProfileEditorVisibility)
 		SLATE_ATTRIBUTE(bool, EditPanelVisible)
@@ -47,6 +48,7 @@ private:
 	CustomLaunchUI::FOnProfileClicked OnProfileDelete;
 	CustomLaunchUI::FOnProfileClicked OnProfileDuplicate;
 	CustomLaunchUI::FOnProfileClicked OnProfileRename;
+	CustomLaunchUI::FOnProfileClicked OnProfileEditDescription;
 	CustomLaunchUI::FOnProfileClicked OnProfileModified;
 	FSetProfileEditorVisible ChangeProfileEditorVisibility;
 	TAttribute<bool> EditPanelVisible;
@@ -63,6 +65,7 @@ private:
 	void OnCustomProfileDuplicateClicked(ILauncherProfilePtr Profile);
 	void OnCustomProfileDeleteClicked(ILauncherProfilePtr Profile);
 	void OnCustomProfileRenameClicked(ILauncherProfilePtr Profile);
+	void OnCustomProfileEditDescriptionClicked(ILauncherProfilePtr Profile);
 	FReply OnOpenDeviceManagerClicked();
 	FReply OnEditProfileClicked(ILauncherProfilePtr Profile);
 	FReply OnCloseEditorClicked();
diff --git a/Engine/Plugins/Developer/CustomLaunchUI/Source/CustomLaunchUI/Private/Widgets/Profiles/SCustomLaunchProfilesPanel.cpp b/Engine/Plugins/Developer/CustomLaunchUI/Source/CustomLaunchUI/Private/Widgets/Profiles/SCustomLaunchProfilesPanel.cpp
--- a/Engine/Plugins/Developer/CustomLaunchUI/Source/CustomLaunchUI/Private/Widgets/Profiles/SCustomLaunchProfilesPanel.cpp
+++ b/Engine/Plugins/Developer/CustomLaunchUI/Source/CustomLaunchUI/Private/Widgets/Profiles/SCustomLaunchProfilesPanel.cpp
@@ -297,6 +297,7 @@ TSharedRef<SWidget> SCustomLaunchProfilesPanel::CreateProfileSelectorWidget()
 		.OnProfileDuplicate(this, &SCustomLaunchProfilesPanel::OnProfileDuplicate)
 		.OnProfileEdit(this, &SCustomLaunchProfilesPanel::OnProfileEdit)
 		.OnProfileRename(this, &SCustomLaunchProfilesPanel::OnProfileRename)
+		.OnProfileEditDescription(this, &SCustomLaunchProfilesPanel::OnProfileEditDescription)
 		.OnProfileModified(this, &SCustomLaunchProfilesPanel::OnProfileModified)
 		.ChangeProfileEditorVisibility(this, &SCustomLaunchProfilesPanel::SetProfileEditorVisible)
 		.EditPanelVisible(this, &SCustomLaunchProfilesPanel::IsProfileEditorVisible)
@@ -393,7 +394,7 @@ TSharedRef<SWidget> SCustomLaunchProfilesPanel::CreateProfileEditorToolbarWidget
 
 			// profile details
 			+SHorizontalBox::Slot()
-			.AutoWidth()
+			.FillWidth(1)
 			.Padding(4)
 			[
 				SNew(SVerticalBox)
@@ -405,6 +406,7 @@ TSharedRef<SWidget> SCustomLaunchProfilesPanel::CreateProfileEditorToolbarWidget
 				[
 					SNew(STextBlock)
 					.Text( this, &SCustomLaunchProfilesPanel::GetSelectedProfileName )
+					.OverflowPolicy(ETextOverflowPolicy::Clip)
 				]
 
 				// profile description
@@ -415,16 +417,10 @@ TSharedRef<SWidget> SCustomLaunchProfilesPanel::CreateProfileEditorToolbarWidget
 					SNew(STextBlock)
 					.Text( this, &SCustomLaunchProfilesPanel::GetSelectedProfileDescription )
 					.ColorAndOpacity(FSlateColor::UseSubduedForeground())
+					.OverflowPolicy(ETextOverflowPolicy::Ellipsis)
 				]
 			]
 
-			// spacer
-			+SHorizontalBox::Slot()
-			.FillWidth(1)
-			[
-				SNew(SSpacer)
-			]
-
 			// launch button
 			+SHorizontalBox::Slot()
 			.AutoWidth()
@@ -575,6 +571,16 @@ void SCustomLaunchProfilesPanel::OnProfileRename(const ILauncherProfilePtr& Prof
 	}
 }
 
+void SCustomLaunchProfilesPanel::OnProfileEditDescription(const ILauncherProfilePtr& Profile)
+{
+	if (Profile.IsValid() && !Model->IsBasicLaunchProfile(Profile))
+	{
+		Model->SelectProfile(Profile);
+
+		ProfileSelector->StartEditProfileDescription(Profile);
+	}
+}
+
 
 void SCustomLaunchProfilesPanel::OnProfileSave(const ILauncherProfilePtr& Profile)
 {
@@ -610,7 +616,40 @@ FText SCustomLaunchProfilesPanel::GetLaunchButtonToolTipText() const
 
 const FSlateBrush* SCustomLaunchProfilesPanel::GetLaunchButtonImage() const
 {
-	return IsLaunchButtonEnabled() ? FAppStyle::Get().GetBrush("Icons.Launch") : FAppStyle::Get().GetBrush("Icons.Error");
+	ILauncherProfilePtr Profile = Model->GetSelectedProfile();
+
+	if (!IsLaunchButtonEnabled())
+	{
+		return FAppStyle::Get().GetBrush("Icons.Error");
+	}
+	else if (Profile->GetLaunchMode() != ELauncherProfileLaunchModes::DoNotLaunch)
+	{
+		return FCustomLaunchUIStyle::Get().GetBrush("Icons.Task.Launch");
+	}
+	else if (Profile->GetDeploymentMode() != ELauncherProfileDeploymentModes::DoNotDeploy)
+	{
+		return FCustomLaunchUIStyle::Get().GetBrush("Icons.Task.Deploy");
+	}
+	else if (Profile->GetPackagingMode() != ELauncherProfilePackagingModes::DoNotPackage)
+	{
+		return FCustomLaunchUIStyle::Get().GetBrush("Icons.Task.Package");
+	}
+	else if (Profile->GetCookMode() != ELauncherProfileCookModes::DoNotCook)
+	{
+		return FCustomLaunchUIStyle::Get().GetBrush("Icons.Task.Cook");
+	}
+	else if (Profile->GetBuildMode() != ELauncherProfileBuildModes::DoNotBuild)
+	{
+		return FCustomLaunchUIStyle::Get().GetBrush("Icons.Task.Build");
+	}
+	else if (Profile->IsImportingZenSnapshot())
+	{
+		return FCustomLaunchUIStyle::Get().GetBrush("Icons.Task.Zen");
+	}
+	else
+	{
+		return FCustomLaunchUIStyle::Get().GetBrush("Icons.Task.Launch");
+	}
 }
 
 
diff --git a/Engine/Plugins/Developer/CustomLaunchUI/Source/CustomLaunchUI/Private/Widgets/Profiles/SCustomLaunchProfilesPanel.h b/Engine/Plugins/Developer/CustomLaunchUI/Source/CustomLaunchUI/Private/Widgets/Profiles/SCustomLaunchProfilesPanel.h
--- a/Engine/Plugins/Developer/CustomLaunchUI/Source/CustomLaunchUI/Private/Widgets/Profiles/SCustomLaunchProfilesPanel.h
+++ b/Engine/Plugins/Developer/CustomLaunchUI/Source/CustomLaunchUI/Private/Widgets/Profiles/SCustomLaunchProfilesPanel.h
@@ -42,6 +42,7 @@ private:
 	void OnProfileSave(const ILauncherProfilePtr& Profile);
 	void OnProfileDuplicate(const ILauncherProfilePtr& Profile);
 	void OnProfileRename(const ILauncherProfilePtr& Profile);
+	void OnProfileEditDescription(const ILauncherProfilePtr& Profile);
 	void OnProfileSelected(const ILauncherProfilePtr& NewProfile, const ILauncherProfilePtr& OldProfile);
 	void OnProfileModified(const ILauncherProfilePtr& Profile);
 
diff --git a/Engine/Source/Runtime/Slate/Private/Widgets/Text/SInlineEditableTextBlock.cpp b/Engine/Source/Runtime/Slate/Private/Widgets/Text/SInlineEditableTextBlock.cpp
--- a/Engine/Source/Runtime/Slate/Private/Widgets/Text/SInlineEditableTextBlock.cpp
+++ b/Engine/Source/Runtime/Slate/Private/Widgets/Text/SInlineEditableTextBlock.cpp
@@ -18,6 +18,7 @@ void SInlineEditableTextBlock::Construct( const FArguments& InArgs )
 	IsSelected = InArgs._IsSelected;
 	OnVerifyTextChanged = InArgs._OnVerifyTextChanged;
 	Text = InArgs._Text;
+	HintText = InArgs._HintText;
 	bIsReadOnly = InArgs._IsReadOnly;
 	MaximumLength = InArgs._MaximumLength;
 
@@ -57,6 +58,7 @@ void SInlineEditableTextBlock::Construct( const FArguments& InArgs )
 	{
 		SAssignNew(MultiLineTextBox, SMultiLineEditableTextBox)
 			.Text(InArgs._Text)
+			.HintText(InArgs._HintText)
 			.Style(&InArgs._Style->EditableTextBoxStyle)
 			.Font(InArgs._Font)
 			.ToolTipText(InArgs._ToolTipText)
@@ -76,6 +78,7 @@ void SInlineEditableTextBlock::Construct( const FArguments& InArgs )
 	{
 		SAssignNew(TextBox, SEditableTextBox)
 			.Text(InArgs._Text)
+			.HintText(InArgs._HintText)
 			.Style(&InArgs._Style->EditableTextBoxStyle)
 			.Font(InArgs._Font)
 			.ToolTipText(InArgs._ToolTipText)
@@ -412,6 +415,16 @@ TSharedPtr<SWidget> SInlineEditableTextBlock::GetEditableTextWidget() const
 	}
 }
 
+void SInlineEditableTextBlock::SetHintText( const TAttribute< FText >& InHintText )
+{
+	HintText = InHintText;
+#if WITH_FANCY_TEXT
+	bIsMultiLine ? MultiLineTextBox->SetHintText(InHintText) :
+#endif //WITH_FANCY_TEXT
+	TextBox->SetHintText(InHintText);
+}
+
+
 void SInlineEditableTextBlock::SetEditableText( const TAttribute< FText >& InNewText )
 {
 #if WITH_FANCY_TEXT
diff --git a/Engine/Source/Runtime/Slate/Public/Widgets/Text/SInlineEditableTextBlock.h b/Engine/Source/Runtime/Slate/Public/Widgets/Text/SInlineEditableTextBlock.h
--- a/Engine/Source/Runtime/Slate/Public/Widgets/Text/SInlineEditableTextBlock.h
+++ b/Engine/Source/Runtime/Slate/Public/Widgets/Text/SInlineEditableTextBlock.h
@@ -32,6 +32,7 @@ class SInlineEditableTextBlock: public SCompoundWidget
 {
 	SLATE_BEGIN_ARGS( SInlineEditableTextBlock )
 		: _Text()
+		, _HintText()
 		, _Style( &FCoreStyle::Get().GetWidgetStyle<FInlineEditableTextBlockStyle>("InlineEditableTextBlockStyle") )
 		, _Font()
 		, _ColorAndOpacity()
@@ -55,6 +56,9 @@ class SInlineEditableTextBlock: public SCompoundWidget
 		/** The text displayed in this text block */
 		SLATE_ATTRIBUTE( FText, Text )
 
+		/** Hint text that appears when there is no text in the text box */
+		SLATE_ATTRIBUTE( FText, HintText )
+
 		/** Pointer to a style of the inline editable text block, which dictates the font, color, and shadow options. */
 		SLATE_STYLE_ARGUMENT( FInlineEditableTextBlockStyle, Style )
 
@@ -178,6 +182,8 @@ class SInlineEditableTextBlock: public SCompoundWidget
 	/** Return the current text */
 	SLATE_API FText GetText() const;
 
+	SLATE_API void SetHintText( const TAttribute< FText >& InHintText );
+
 	SLATE_API void SetHighlightText(const TAttribute<FText>& InText);
 
 	/** Sets the wrap text at attribute.  See WrapTextAt attribute */
@@ -243,6 +249,9 @@ protected:
 	/** Attribute for the text to use for the widget */
 	TAttribute<FText> Text;
 
+	/** Attribute for the hint text to use for the widget */
+	TAttribute<FText> HintText;
+
 	/** Attribute to look up if the widget is read-only */
 	TAttribute< bool > bIsReadOnly;
 
