Project Launcher extension changes:
- Split custom Extension menu items to a separate drop-down to keep the main 'add parameter' menu cleaner (this separate drop-down is temporary and these extension items will move to a dedicated "extensions" button on the bottom toolbar)
- Extension system disabled by default, hidden behind experimental feature flag. (note: the flag is not honored by UnrealFrontend so extensions are always unavailable)
- Remove the old experimental flag that enabled Project Launcher 2.
- Fix the Insights extension so it won't show an empty channel submenu if Insights is unavailable + enabling Insights for UnrealFrontend so the menu is populated

#rnx

diff --git a/Engine/Config/BaseEditorPerProjectUserSettings.ini b/Engine/Config/BaseEditorPerProjectUserSettings.ini
--- a/Engine/Config/BaseEditorPerProjectUserSettings.ini
+++ b/Engine/Config/BaseEditorPerProjectUserSettings.ini
@@ -156,7 +156,6 @@ bToolbarCustomization=False
 bBehaviorTreeEditor=False
 bEnableFindAndReplaceReferences =False
 bExampleLayersAndBlends = True
-bShowNewLaunchProfileTools=True
 
 ; Enabled by default now
 bEnableAsyncTextureCompilation=True
diff --git a/Engine/Plugins/Developer/CustomLaunchUI/Source/CommonLaunchExtensions/Private/Insights/InsightsLaunchExtension.cpp b/Engine/Plugins/Developer/CustomLaunchUI/Source/CommonLaunchExtensions/Private/Insights/InsightsLaunchExtension.cpp
--- a/Engine/Plugins/Developer/CustomLaunchUI/Source/CommonLaunchExtensions/Private/Insights/InsightsLaunchExtension.cpp
+++ b/Engine/Plugins/Developer/CustomLaunchUI/Source/CommonLaunchExtensions/Private/Insights/InsightsLaunchExtension.cpp
@@ -80,13 +80,16 @@ bool FInsightsLaunchExtensionInstance::IsTraceChannelEnabled( const FString InCh
 }
 
 
-void FInsightsLaunchExtensionInstance::CustomizeSubmenu( FMenuBuilder& MenuBuilder )
+void FInsightsLaunchExtensionInstance::CustomizeParametersSubmenu( FMenuBuilder& MenuBuilder )
 {
+#if UE_TRACE_ENABLED || UE_TRACE_MINIMAL_ENABLED
 	auto ChannelsMenuDelegate = [this]( FMenuBuilder& MenuBuilder )
 	{
 		CacheTraceChannels();
 
 		// collect all trace channels
+		// note that this will enumerate the channels that are available for the editor / UnrealFrontend, not the ones available to the game
+		// UnrealFrontend does not have all channels available
 		TArray<FString> AllTraceChannels;
 		UE::Trace::EnumerateChannels([](const UE::Trace::FChannelInfo& ChannelInfo, void* User)
 		{
@@ -141,6 +144,7 @@ void FInsightsLaunchExtensionInstance::CustomizeSubmenu( FMenuBuilder& MenuBuild
 		FSlateIcon(),
 		false // bInShouldCloseWindowAfterMenuSelection
 	);
+#endif
 }
 
 
diff --git a/Engine/Plugins/Developer/CustomLaunchUI/Source/CommonLaunchExtensions/Private/Insights/InsightsLaunchExtension.h b/Engine/Plugins/Developer/CustomLaunchUI/Source/CommonLaunchExtensions/Private/Insights/InsightsLaunchExtension.h
--- a/Engine/Plugins/Developer/CustomLaunchUI/Source/CommonLaunchExtensions/Private/Insights/InsightsLaunchExtension.h
+++ b/Engine/Plugins/Developer/CustomLaunchUI/Source/CommonLaunchExtensions/Private/Insights/InsightsLaunchExtension.h
@@ -15,7 +15,7 @@ public:
 	virtual bool GetExtensionParameters( TArray<FString>& OutParameters ) const override;
 	virtual FText GetExtensionParameterDisplayName( const FString& InParameter ) const override;
 	
-	virtual void CustomizeSubmenu( FMenuBuilder& MenuBuilder ) override;
+	virtual void CustomizeParametersSubmenu( FMenuBuilder& MenuBuilder ) override;
 
 	void CacheTraceChannels();
 	void ToggleTraceChannel( const FString& InChannel );
diff --git a/Engine/Plugins/Developer/CustomLaunchUI/Source/CustomLaunchUI/Private/Extension/LaunchExtension.cpp b/Engine/Plugins/Developer/CustomLaunchUI/Source/CustomLaunchUI/Private/Extension/LaunchExtension.cpp
--- a/Engine/Plugins/Developer/CustomLaunchUI/Source/CustomLaunchUI/Private/Extension/LaunchExtension.cpp
+++ b/Engine/Plugins/Developer/CustomLaunchUI/Source/CustomLaunchUI/Private/Extension/LaunchExtension.cpp
@@ -368,7 +368,7 @@ namespace CustomLaunchUI
 			}
 		}
 
-		CustomizeSubmenu(MenuBuilder);
+		CustomizeParametersSubmenu(MenuBuilder);
 	}
 
 	FString FLaunchExtensionInstance::GetCommandLine() const
diff --git a/Engine/Plugins/Developer/CustomLaunchUI/Source/CustomLaunchUI/Private/Marshalling/LaunchProfileTreeData.cpp b/Engine/Plugins/Developer/CustomLaunchUI/Source/CustomLaunchUI/Private/Marshalling/LaunchProfileTreeData.cpp
--- a/Engine/Plugins/Developer/CustomLaunchUI/Source/CustomLaunchUI/Private/Marshalling/LaunchProfileTreeData.cpp
+++ b/Engine/Plugins/Developer/CustomLaunchUI/Source/CustomLaunchUI/Private/Marshalling/LaunchProfileTreeData.cpp
@@ -18,6 +18,7 @@
 #include "Framework/MultiBox/MultiBoxBuilder.h"
 #include "Framework/Commands/UIAction.h"
 #include "DesktopPlatformModule.h"
+#include "Misc/ConfigCacheIni.h"
 
 
 #define LOCTEXT_NAMESPACE "SCustomLaunchCustomProfileEditor"
@@ -34,6 +35,15 @@ namespace CustomLaunchUI
 		if (InProfile.IsValid())
 		{
 			ExtensionInstances = FLaunchExtension::CreateExtensionInstancesForProfile(Profile.ToSharedRef(), InModel, InMarshaller);
+
+			bHasAnyMenuExtensions = false;
+			for (TSharedPtr<FLaunchExtensionInstance> ExtensionInstance : ExtensionInstances)
+			{
+				if (ExtensionInstance->HasCustomExtensionMenu())
+				{
+					bHasAnyMenuExtensions = true;
+				}
+			}
 		}
 	}
 
@@ -270,6 +280,19 @@ namespace CustomLaunchUI
 
 	FLaunchProfileTreeNode& FLaunchProfileTreeNode::AddCommandLineString( FText InName, FStringCallbacks&& StringCallbacks )
 	{
+		// make sure the extensions system is enabled
+		// ...unfortunately this setting is only available in the editor because the ini file is not loaded by UnrealFrontend because its per-project
+		bool bEnableProjectLauncherExtensions = false;
+		GConfig->GetBool(TEXT("/Script/UnrealEd.EditorExperimentalSettings"), TEXT("bEnableProjectLauncherExtensions"), bEnableProjectLauncherExtensions, GEditorPerProjectIni);
+
+		if (!bEnableProjectLauncherExtensions)
+		{
+			return AddString(InName, MoveTemp(StringCallbacks));
+		}
+
+
+
+
 		check(StringCallbacks.GetValue);
 		check(StringCallbacks.SetValue);
 
@@ -303,7 +326,7 @@ namespace CustomLaunchUI
 			return FText::FromString(StringCallbacks.GetValue());
 		};
 
-		auto OnGetMenuContent = [this,StringCallbacks]()
+		auto OnGetCmdlineParameterMenuContent = [this,StringCallbacks]()
 		{
 			const bool bShouldCloseWindowAfterMenuSelection = false;
 			const bool bCloseSelfOnly = false;
@@ -334,6 +357,44 @@ namespace CustomLaunchUI
 			return MenuBuilder.MakeWidget();
 		};
 
+		auto OnGetExtensionsMenuContent = [this,StringCallbacks]()
+		{
+			const bool bShouldCloseWindowAfterMenuSelection = false;
+			const bool bCloseSelfOnly = false;
+			const bool bSearchable = false;
+			FMenuBuilder MenuBuilder(bShouldCloseWindowAfterMenuSelection, nullptr, nullptr, bCloseSelfOnly, &FCoreStyle::Get(), bSearchable );
+
+			for (TSharedPtr<FLaunchExtensionInstance> ExtensionInstance : TreeData->ExtensionInstances)
+			{
+				if (!ExtensionInstance->HasCustomExtensionMenu())
+				{
+					continue;
+				}
+
+#if 1
+				// flat list of extensions - easier to see but risks getting cluttered later on
+				MenuBuilder.BeginSection(NAME_None, ExtensionInstance->GetExtension()->GetDisplayName());
+				ExtensionInstance->MakeCustomExtensionSubmenu(MenuBuilder);
+				MenuBuilder.EndSection();
+#else 
+				// extensions in submenus - more scalable, but harder to see everything at a glance
+				const bool bInOpenSubMenuOnClick = false;
+				MenuBuilder.AddSubMenu(
+					ExtensionInstance->GetExtension()->GetDisplayName()),
+					FText::GetEmpty(),
+					FNewMenuDelegate::CreateSP( ExtensionInstance.ToSharedRef(), &FLaunchExtensionInstance::MakeExtensionSubmenu),
+					bInOpenSubMenuOnClick,
+					FSlateIcon(),
+					bShouldCloseWindowAfterMenuSelection
+				);
+#endif
+			}
+
+			return MenuBuilder.MakeWidget();
+		};
+
+
+
 
 		return AddWidget( InName,
 			MoveTemp(WidgetCallbacks),
@@ -352,24 +413,45 @@ namespace CustomLaunchUI
 				.OnTextCommitted_Lambda( SetString )
 			]
 
-			// command line extension button
+			// command line parameters button
 			+SHorizontalBox::Slot()
 			.AutoWidth()
-			.Padding(FMargin(0, 0, 4, 0))
+			.Padding(FMargin(4, 0, 0, 0))
 			.VAlign(VAlign_Center)
 			[
 				SNew(SComboButton)
 				.ComboButtonStyle(FAppStyle::Get(), "SimpleComboButton")
-				.ToolTipText(LOCTEXT("CmdLineOptionsLabel", "Add a special item to command line"))
-				.OnGetMenuContent_Lambda(OnGetMenuContent)
+				.ToolTipText(LOCTEXT("CmdLineOptionsLabel", "Add a special parameter to command line"))
+				.OnGetMenuContent_Lambda(OnGetCmdlineParameterMenuContent)
 				.Visibility_Lambda( [this]() { return TreeData->ExtensionInstances.Num() > 0 ? EVisibility::Visible : EVisibility::Collapsed; } )
 				.MenuPlacement(MenuPlacement_BelowRightAnchor)
 				.HasDownArrow(false)
 				.ButtonContent()
 				[
 					SNew(SImage)
-					.Image(FAppStyle::Get().GetBrush("Icons.Plus"))
-					.ColorAndOpacity(FStyleColors::AccentGreen)
+					.Image(FAppStyle::Get().GetBrush("Icons.AddCircle"))
+					.ColorAndOpacity(FSlateColor::UseForeground())
+				]
+			]
+
+			// command line extension button (@todo: temporary! this will be moved elsewhere)
+			+SHorizontalBox::Slot()
+			.AutoWidth()
+			.Padding(FMargin(0, 0, 3, 0))
+			.VAlign(VAlign_Center)
+			[
+				SNew(SComboButton)
+				.ComboButtonStyle(FAppStyle::Get(), "SimpleComboButton")
+				.ToolTipText(LOCTEXT("CmdLineExtensionsLabel", "Show additional command line extension options"))
+				.OnGetMenuContent_Lambda(OnGetExtensionsMenuContent)
+				.Visibility_Lambda( [this]() { return TreeData->bHasAnyMenuExtensions ? EVisibility::Visible : EVisibility::Collapsed; } )
+				.MenuPlacement(MenuPlacement_BelowRightAnchor)
+				.HasDownArrow(false)
+				.ButtonContent()
+				[
+					SNew(SImage)
+					.Image(FCustomLaunchUIStyle::Get().GetBrush("Icons.EllipsisVerticalNarrow"))
+					.ColorAndOpacity(FSlateColor::UseForeground())
 				]
 			]
 		);
diff --git a/Engine/Plugins/Developer/CustomLaunchUI/Source/CustomLaunchUI/Private/Styling/CustomLaunchUIStyle.cpp b/Engine/Plugins/Developer/CustomLaunchUI/Source/CustomLaunchUI/Private/Styling/CustomLaunchUIStyle.cpp
--- a/Engine/Plugins/Developer/CustomLaunchUI/Source/CustomLaunchUI/Private/Styling/CustomLaunchUIStyle.cpp
+++ b/Engine/Plugins/Developer/CustomLaunchUI/Source/CustomLaunchUI/Private/Styling/CustomLaunchUIStyle.cpp
@@ -73,6 +73,8 @@ void FCustomLaunchUIStyle::Initialize()
 
 	StyleSet->Set("Icons.DiffersFromDefault", new CORE_IMAGE_BRUSH_SVG( "../Editor/Slate/Starship/Common/ResetToDefault", Icon16x16) );
 	StyleSet->Set("Icons.DeviceManager", new CORE_IMAGE_BRUSH_SVG( "../Editor/Slate/Starship/Common/DeviceManager", Icon16x16) );
+	StyleSet->Set("Icons.EllipsisVerticalNarrow", new CORE_IMAGE_BRUSH_SVG("Starship/Common/ellipsis-vertical-narrow", FVector2f(6.f, 20.f))); // @todo: temporary - extension menu will move elsewhere
+
 	StyleSet->Set("Icons.Asset", new CORE_IMAGE_BRUSH( "Icons/doc_16x", Icon16x16 ) );
 
 	StyleSet->Set( "PathPickerButton", new CORE_IMAGE_BRUSH( "Icons/ellipsis_12x", Icon12x12 ) );
diff --git a/Engine/Plugins/Developer/CustomLaunchUI/Source/CustomLaunchUI/Public/Extension/LaunchExtension.h b/Engine/Plugins/Developer/CustomLaunchUI/Source/CustomLaunchUI/Public/Extension/LaunchExtension.h
--- a/Engine/Plugins/Developer/CustomLaunchUI/Source/CustomLaunchUI/Public/Extension/LaunchExtension.h
+++ b/Engine/Plugins/Developer/CustomLaunchUI/Source/CustomLaunchUI/Public/Extension/LaunchExtension.h
@@ -57,9 +57,9 @@ namespace CustomLaunchUI
 		virtual bool GetExtensionVariableValue( const FString& InParameter, FString& OutValue ) const;
 
 		/**
-		 * Hook to allow the extension to extend its menu
+		 * Hook to allow the extension to extend the extension parameters menu
 		 */
-		virtual void CustomizeSubmenu( FMenuBuilder& MenuBuilder ) {};
+		virtual void CustomizeParametersSubmenu( FMenuBuilder& MenuBuilder ) {};
 
 		/**
 		 * Hook to allow the extension to add extra fields to the property editing tree, if the marshall allows it
@@ -72,6 +72,19 @@ namespace CustomLaunchUI
 		 */
 		virtual void CustomizeLaunchCommandLine( FString& InOutCommandLine ) {};
 
+
+		/**
+		 * Determine if this extension also provides items in the bespoke extensions menu
+		 */
+		virtual bool HasCustomExtensionMenu() const { return false; }
+
+		/*
+		 * Populate the custom extension menu for this extension. This menu would typically contain
+		 * items that might enable bespoke tree customization options for example.
+		 */
+		virtual void MakeCustomExtensionSubmenu(FMenuBuilder& MenuBuilder) { checkNoEntry(); }
+
+
 	protected:
 
 		/**
diff --git a/Engine/Plugins/Developer/CustomLaunchUI/Source/CustomLaunchUI/Public/Marshalling/LaunchProfileTreeData.h b/Engine/Plugins/Developer/CustomLaunchUI/Source/CustomLaunchUI/Public/Marshalling/LaunchProfileTreeData.h
--- a/Engine/Plugins/Developer/CustomLaunchUI/Source/CustomLaunchUI/Public/Marshalling/LaunchProfileTreeData.h
+++ b/Engine/Plugins/Developer/CustomLaunchUI/Source/CustomLaunchUI/Public/Marshalling/LaunchProfileTreeData.h
@@ -106,6 +106,7 @@ namespace CustomLaunchUI
 		TArray<FLaunchProfileTreeNodePtr> Nodes;
 		ILaunchProfileMarshaller* Marshaller;
 		TArray<TSharedPtr<FLaunchExtensionInstance>> ExtensionInstances;
+		bool bHasAnyMenuExtensions = false;
 	};
 
 
diff --git a/Engine/Source/Editor/UnrealEd/Classes/Settings/EditorExperimentalSettings.h b/Engine/Source/Editor/UnrealEd/Classes/Settings/EditorExperimentalSettings.h
--- a/Engine/Source/Editor/UnrealEd/Classes/Settings/EditorExperimentalSettings.h
+++ b/Engine/Source/Editor/UnrealEd/Classes/Settings/EditorExperimentalSettings.h
@@ -120,9 +120,9 @@ public:
 	UPROPERTY(EditAnywhere, config, Category = Tools, meta = (DisplayName = "Enable tools in Packaging Settings -> Additional builds for this project"))
 	bool bProjectCustomBuildTools;
 
-	/** Shows an option to configure custom launch profiles for advanced packaging, deploying and launching of your projects */
+	/** Enables Extensions for the new Project Launcher */
 	UPROPERTY(EditAnywhere, config, Category = Tools)
-	bool bShowNewLaunchProfileTools;
+	bool bEnableProjectLauncherExtensions;
 
 	/** Enable experimental PIE preview device launch */
 	UPROPERTY(EditAnywhere, config, Category = PIE, meta = (DisplayName = "Enable mobile PIE with preview device launch options."))
diff --git a/Engine/Source/Programs/UnrealFrontend/UnrealFrontend.Target.cs b/Engine/Source/Programs/UnrealFrontend/UnrealFrontend.Target.cs
--- a/Engine/Source/Programs/UnrealFrontend/UnrealFrontend.Target.cs
+++ b/Engine/Source/Programs/UnrealFrontend/UnrealFrontend.Target.cs
@@ -26,6 +26,8 @@ public class UnrealFrontendTarget : TargetRules
 
 		bHasExports = false;
 
+		bEnableTrace = true;
+
 		// Old Profiler (SessionFrontend/Profiler) is deprecated since UE 5.0. Use Trace/UnrealInsights instead.
 		//GlobalDefinitions.Add("UE_DEPRECATED_PROFILER_ENABLED=1");
 
